Fix issue when World Partition doesn't have a top-view image (show actors bounding volumes in this case).
Also show world's bounds in the world partiton editor.

diff --git a/Engine/Source/Editor/WorldPartitionEditor/Private/WorldPartition/SWorldPartitionEditorGrid2D.cpp b/Engine/Source/Editor/WorldPartitionEditor/Private/WorldPartition/SWorldPartitionEditorGrid2D.cpp
--- a/Engine/Source/Editor/WorldPartitionEditor/Private/WorldPartition/SWorldPartitionEditorGrid2D.cpp
+++ b/Engine/Source/Editor/WorldPartitionEditor/Private/WorldPartition/SWorldPartitionEditorGrid2D.cpp
@@ -61,6 +61,10 @@ void SWorldPartitionEditorGrid2D::Construct(const FArguments& InArgs)
 {
 	SWorldPartitionEditorGrid::Construct(SWorldPartitionEditorGrid::FArguments().InWorld(InArgs._InWorld));
 
+	// Defaults
+	Trans = FVector2D(0, 0);
+	Scale = 0.00133333332;
+
 	// UI
 	ChildSlot
 	[
@@ -601,7 +605,7 @@ uint32 SWorldPartitionEditorGrid2D::PaintScaleRuler(const FGeometry& AllottedGeo
 		ESlateDrawEffect::None,
 		FLinearColor::White);
 
-	const float UnitsInRuler = ScaleRulerLength/Scale + 0.05f;// Pixels to world units (+0.05f to accomodate for %.2f)
+	const float UnitsInRuler = ScaleRulerLength/Scale + 0.05f;// Pixels to world units (+0.05f to accommodate for %.2f)
 	const int32 UnitsInMeter = 100;
 	const int32 UnitsInKilometer = UnitsInMeter*1000;
 	
@@ -623,6 +627,21 @@ uint32 SWorldPartitionEditorGrid2D::PaintScaleRuler(const FGeometry& AllottedGeo
 		FEditorStyle::GetFontStyle("NormalFont"),
 		ESlateDrawEffect::None,
 		FLinearColor::White);
+
+	// Show world bounds
+	const FBox WorldBounds = WorldPartition->GetWorldBounds();
+	const FVector WorldBoundsExtentInKM = (WorldBounds.GetExtent() * 2.0f) / 100000.0f;
+	RulerText = FString::Printf(TEXT("%.2fx%.2fx%.2f km"), WorldBoundsExtentInKM.X, WorldBoundsExtentInKM.Y, WorldBoundsExtentInKM.Z);
+	
+	FSlateDrawElement::MakeText(
+		OutDrawElements,
+		LayerId,
+		AllottedGeometry.ToOffsetPaintGeometry(FVector2D(10, 67)),
+		RulerText,
+		FEditorStyle::GetFontStyle("NormalFont"),
+		ESlateDrawEffect::None,
+		FLinearColor::White);
+
 		
 	return LayerId + 1;
 }
diff --git a/Engine/Source/Editor/WorldPartitionEditor/Private/WorldPartition/SWorldPartitionEditorGridSpatialHash.cpp b/Engine/Source/Editor/WorldPartitionEditor/Private/WorldPartition/SWorldPartitionEditorGridSpatialHash.cpp
--- a/Engine/Source/Editor/WorldPartitionEditor/Private/WorldPartition/SWorldPartitionEditorGridSpatialHash.cpp
+++ b/Engine/Source/Editor/WorldPartitionEditor/Private/WorldPartition/SWorldPartitionEditorGridSpatialHash.cpp
@@ -34,7 +34,8 @@ SWorldPartitionEditorGridSpatialHash::~SWorldPartitionEditorGridSpatialHash()
 
 void SWorldPartitionEditorGridSpatialHash::Construct(const FArguments& InArgs)
 {
-	SWorldPartitionEditorGrid2D::Construct(SWorldPartitionEditorGrid::FArguments().InWorld(InArgs._InWorld));
+	World = InArgs._InWorld;
+	WorldPartition = World ? World->GetWorldPartition() : nullptr;
 
 	UWorldPartitionEditorSpatialHash* EditorSpatialHash = (UWorldPartitionEditorSpatialHash*)WorldPartition->EditorHash;
 
@@ -54,8 +55,7 @@ void SWorldPartitionEditorGridSpatialHash::Construct(const FArguments& InArgs)
 		}
 	}
 
-	Trans = FVector2D(0, 0);
-	Scale = 0.00133333332;
+	SWorldPartitionEditorGrid2D::Construct(SWorldPartitionEditorGrid::FArguments().InWorld(InArgs._InWorld));
 }
 
 int32 SWorldPartitionEditorGridSpatialHash::PaintGrid(const FGeometry& AllottedGeometry, const FSlateRect& MyCullingRect, FSlateWindowElementList& OutDrawElements, int32 LayerId) const
