Fix Win64 static analysis warning.

diff --git a/Engine/Source/Editor/WorldPartitionEditor/Private/WorldPartition/SWorldPartitionEditorGridSpatialHash.cpp b/Engine/Source/Editor/WorldPartitionEditor/Private/WorldPartition/SWorldPartitionEditorGridSpatialHash.cpp
--- a/Engine/Source/Editor/WorldPartitionEditor/Private/WorldPartition/SWorldPartitionEditorGridSpatialHash.cpp
+++ b/Engine/Source/Editor/WorldPartitionEditor/Private/WorldPartition/SWorldPartitionEditorGridSpatialHash.cpp
@@ -37,21 +37,24 @@ void SWorldPartitionEditorGridSpatialHash::Construct(const FArguments& InArgs)
 	World = InArgs._InWorld;
 	WorldPartition = World ? World->GetWorldPartition() : nullptr;
 
-	UWorldPartitionEditorSpatialHash* EditorSpatialHash = (UWorldPartitionEditorSpatialHash*)WorldPartition->EditorHash;
-
-	bShowActors = true;
-	if (UObject* WorldImage = EditorSpatialHash->WorldImage.TryLoad())
+	if (WorldPartition)
 	{
-		if (UTexture2D* WorldImageTexture = Cast<UTexture2D>(WorldImage))
-		{
-			WorldImageBrush.SetImageSize(FVector2D(WorldImageTexture->GetSizeX(), WorldImageTexture->GetSizeY()));
-			WorldImageBrush.SetResourceObject(WorldImageTexture);
-			bShowActors = false;
-		}
-		if (UMaterialInterface* WorldImageMaterial = Cast<UMaterialInterface>(WorldImage))
+		UWorldPartitionEditorSpatialHash* EditorSpatialHash = (UWorldPartitionEditorSpatialHash*)WorldPartition->EditorHash;
+
+		bShowActors = true;
+		if (UObject* WorldImage = EditorSpatialHash->WorldImage.TryLoad())
 		{
-			WorldImageBrush.SetResourceObject(WorldImageMaterial);
-			bShowActors = false;
+			if (UTexture2D* WorldImageTexture = Cast<UTexture2D>(WorldImage))
+			{
+				WorldImageBrush.SetImageSize(FVector2D(WorldImageTexture->GetSizeX(), WorldImageTexture->GetSizeY()));
+				WorldImageBrush.SetResourceObject(WorldImageTexture);
+				bShowActors = false;
+			}
+			if (UMaterialInterface* WorldImageMaterial = Cast<UMaterialInterface>(WorldImage))
+			{
+				WorldImageBrush.SetResourceObject(WorldImageMaterial);
+				bShowActors = false;
+			}
 		}
 	}
 
