WorldPartition: Prevent Loading/Unloading Editor Cells outside of Default Editor Modes

diff --git a/Engine/Source/Editor/WorldPartitionEditor/Private/WorldPartition/SWorldPartitionEditorGrid2D.cpp b/Engine/Source/Editor/WorldPartitionEditor/Private/WorldPartition/SWorldPartitionEditorGrid2D.cpp
--- a/Engine/Source/Editor/WorldPartitionEditor/Private/WorldPartition/SWorldPartitionEditorGrid2D.cpp
+++ b/Engine/Source/Editor/WorldPartitionEditor/Private/WorldPartition/SWorldPartitionEditorGrid2D.cpp
@@ -15,6 +15,8 @@
 #include "WorldBrowserModule.h"
 #include "LevelEditorViewport.h"
 #include "Algo/Transform.h"
+#include "Editor.h"
+#include "EditorModeManager.h"
 
 #define LOCTEXT_NAMESPACE "WorldPartitionEditor"
 
@@ -148,13 +150,13 @@ void SWorldPartitionEditorGrid2D::Construct(const FArguments& InArgs)
 	const FEditorCommands& Commands = FEditorCommands::Get();
 	FUICommandList& ActionList = *CommandList;
 
-	auto HasSelectedCells = [this]()
+	auto CanLoadOrUnloadCells = [this]()
 	{
-		return SelectedCells.Num() > 0;
+		return GLevelEditorModeTools().IsDefaultModeActive() && SelectedCells.Num() > 0;
 	};
 
-	ActionList.MapAction(Commands.LoadSelectedCells, FExecuteAction::CreateSP(this, &SWorldPartitionEditorGrid2D::LoadSelectedCells), FCanExecuteAction::CreateLambda(HasSelectedCells));
-	ActionList.MapAction(Commands.UnloadSelectedCells, FExecuteAction::CreateSP(this, &SWorldPartitionEditorGrid2D::UnloadSelectedCells), FCanExecuteAction::CreateLambda(HasSelectedCells));
+	ActionList.MapAction(Commands.LoadSelectedCells, FExecuteAction::CreateSP(this, &SWorldPartitionEditorGrid2D::LoadSelectedCells), FCanExecuteAction::CreateLambda(CanLoadOrUnloadCells));
+	ActionList.MapAction(Commands.UnloadSelectedCells, FExecuteAction::CreateSP(this, &SWorldPartitionEditorGrid2D::UnloadSelectedCells), FCanExecuteAction::CreateLambda(CanLoadOrUnloadCells));
 	ActionList.MapAction(Commands.MoveCameraHere, FExecuteAction::CreateSP(this, &SWorldPartitionEditorGrid2D::MoveCameraHere));
 }
 
