Minor world partition editor UI changes.

diff --git a/Engine/Source/Editor/WorldPartitionEditor/Private/WorldPartition/SWorldPartitionEditorGrid2D.cpp b/Engine/Source/Editor/WorldPartitionEditor/Private/WorldPartition/SWorldPartitionEditorGrid2D.cpp
--- a/Engine/Source/Editor/WorldPartitionEditor/Private/WorldPartition/SWorldPartitionEditorGrid2D.cpp
+++ b/Engine/Source/Editor/WorldPartitionEditor/Private/WorldPartition/SWorldPartitionEditorGrid2D.cpp
@@ -407,7 +407,7 @@ uint32 SWorldPartitionEditorGrid2D::PaintActors(const FGeometry& AllottedGeometr
 				const float SquaredDistanceToPoint = ActorViewBox.ComputeSquaredDistanceToPoint(MouseCursorPos);
 				if ((ActorDesc->Tag == FWorldPartitionActorDesc::GlobalTag) || (SquaredDistanceToPoint > 0.0f && SquaredDistanceToPoint <= 2.0f))
 				{
-					ActorColor = FLinearColor::White;
+					ActorColor = FLinearColor::Yellow;
 
 					FString ActorPath = ActorDesc->GetActorPath().ToString();
 
@@ -427,8 +427,7 @@ uint32 SWorldPartitionEditorGrid2D::PaintActors(const FGeometry& AllottedGeometr
 						ActorColor
 					);
 				}
-
-				if (SelectBox.Intersect(ActorDesc->GetBounds()))
+				else if ((SelectBox.GetVolume() > 0) && SelectBox.Intersect(ActorDesc->GetBounds()))
 				{
 					ActorColor = FLinearColor::White;
 				}
