WorldPartition: New in Editor Conversion process (needs to be enabled in WorldPartition project settings)

diff --git a/Engine/Source/Editor/UnrealEd/Private/FileHelpers.cpp b/Engine/Source/Editor/UnrealEd/Private/FileHelpers.cpp
--- a/Engine/Source/Editor/UnrealEd/Private/FileHelpers.cpp
+++ b/Engine/Source/Editor/UnrealEd/Private/FileHelpers.cpp
@@ -69,6 +69,7 @@
 #include "StudioAnalytics.h"
 #include "AnalyticsEventAttribute.h"
 #include "HierarchicalLOD.h"
+#include "WorldPartition/IWorldPartitionEditorModule.h"
 
 DEFINE_LOG_CATEGORY_STATIC(LogFileHelpers, Log, All);
 
@@ -2545,6 +2546,15 @@ bool FEditorFileUtils::LoadMap(const FString& InFilename, bool LoadAsTemplate, b
 		return false;
 	}
 
+	if (!LoadAsTemplate)
+	{
+		IWorldPartitionEditorModule& WorldPartitionEditorModule = FModuleManager::LoadModuleChecked<IWorldPartitionEditorModule>("WorldPartitionEditor");
+		if (!WorldPartitionEditorModule.ConvertMap(LongMapPackageName))
+		{
+			return false;
+		}
+	}
+
 	// Save last opened level name.
 	GConfig->SetString(TEXT("EditorStartup"), TEXT("LastLevel"), *LongMapPackageName, GEditorPerProjectIni);
 
diff --git a/Engine/Source/Editor/WorldPartitionEditor/Private/WorldPartition/SWorldPartitionConvertDialog.cpp b/Engine/Source/Editor/WorldPartitionEditor/Private/WorldPartition/SWorldPartitionConvertDialog.cpp
new file mode 100644
--- /dev/null
+++ b/Engine/Source/Editor/WorldPartitionEditor/Private/WorldPartition/SWorldPartitionConvertDialog.cpp
@@ -0,0 +1,92 @@
+// Copyright Epic Games, Inc. All Rights Reserved.
+#include "SWorldPartitionConvertDialog.h"
+#include "Widgets/Layout/SSpacer.h"
+#include "WorldPartition/WorldPartitionConvertOptions.h"
+
+#define LOCTEXT_NAMESPACE "WorldPartitionConvertDialog"
+
+const FVector2D SWorldPartitionConvertDialog::DEFAULT_WINDOW_SIZE = FVector2D(600, 350);
+
+void SWorldPartitionConvertDialog::Construct(const FArguments& InArgs)
+{
+	ParentWindowPtr = InArgs._ParentWindow.Get();
+	ConvertOptions = InArgs._ConvertOptions.Get();
+	bClickedOk = false;
+
+	FPropertyEditorModule& PropertyEditorModule = FModuleManager::GetModuleChecked<FPropertyEditorModule>("PropertyEditor");
+	FDetailsViewArgs DetailsViewArgs;
+	DetailsViewArgs.bAllowSearch = false;
+	DetailsViewArgs.NameAreaSettings = FDetailsViewArgs::HideNameArea;
+	DetailsViewArgs.bAllowMultipleTopLevelObjects = false;
+
+	DetailsView = PropertyEditorModule.CreateDetailView(DetailsViewArgs);
+	
+	// Set provided objects on SDetailsView
+	DetailsView->SetObject(ConvertOptions.Get(), true);
+
+	this->ChildSlot
+		[
+			SNew(SBorder)
+			.BorderImage(FEditorStyle::GetBrush("ToolPanel.GroupBorder"))
+			[
+				SNew(SVerticalBox)
+				+ SVerticalBox::Slot()
+				.MaxHeight(500.0f)
+				[
+					DetailsView->AsShared()
+				]
+				+ SVerticalBox::Slot()
+				.AutoHeight()
+				[
+					SNew(SHorizontalBox)
+					+ SHorizontalBox::Slot()
+					.FillWidth(1.f)
+					[
+						SNew(SSpacer)
+					]
+					+ SHorizontalBox::Slot()
+					.AutoWidth()
+					.Padding(2.f)
+					[
+						SNew(SButton)
+						.HAlign(HAlign_Center)
+						.IsEnabled(this, &SWorldPartitionConvertDialog::IsOkEnabled)
+						.ContentPadding(FEditorStyle::GetMargin("StandardDialog.ContentPadding"))
+						.OnClicked(this, &SWorldPartitionConvertDialog::OnOkClicked)
+						.Text(LOCTEXT("OkButton", "Ok"))
+					]
+					+ SHorizontalBox::Slot()
+					.AutoWidth()
+					.Padding(2.f)
+					[
+						SNew(SButton)
+						.HAlign(HAlign_Center)
+						.ContentPadding(FEditorStyle::GetMargin("StandardDialog.ContentPadding"))
+						.OnClicked(this, &SWorldPartitionConvertDialog::OnCancelClicked)
+						.Text(LOCTEXT("CancelButton", "Cancel"))
+					]
+				]
+			]
+		];
+}
+
+bool SWorldPartitionConvertDialog::IsOkEnabled() const
+{
+	return ConvertOptions->CommandletClass != nullptr;
+}
+
+FReply SWorldPartitionConvertDialog::OnOkClicked()
+{
+	bClickedOk = true;
+	ParentWindowPtr.Pin()->RequestDestroyWindow();
+	return FReply::Handled();
+}
+
+FReply SWorldPartitionConvertDialog::OnCancelClicked()
+{
+	bClickedOk = false;
+	ParentWindowPtr.Pin()->RequestDestroyWindow();
+	return FReply::Handled();
+}
+
+#undef LOCTEXT_NAMESPACE
diff --git a/Engine/Source/Editor/WorldPartitionEditor/Private/WorldPartition/SWorldPartitionConvertDialog.h b/Engine/Source/Editor/WorldPartitionEditor/Private/WorldPartition/SWorldPartitionConvertDialog.h
new file mode 100644
--- /dev/null
+++ b/Engine/Source/Editor/WorldPartitionEditor/Private/WorldPartition/SWorldPartitionConvertDialog.h
@@ -0,0 +1,57 @@
+// Copyright Epic Games, Inc. All Rights Reserved.
+#pragma once
+
+#include "CoreMinimal.h"
+#include "InputCoreTypes.h"
+#include "Types/SlateEnums.h"
+#include "Styling/SlateColor.h"
+#include "Input/Reply.h"
+#include "Widgets/SWidget.h"
+#include "Widgets/SWindow.h"
+#include "Widgets/DeclarativeSyntaxSupport.h"
+#include "Widgets/SCompoundWidget.h"
+
+//////////////////////////////////////////////////////////////////////////
+// SWorldPartitionConvertDialog
+class UWorldPartitionConvertOptions;
+
+class SWorldPartitionConvertDialog : public SCompoundWidget
+{
+public:
+
+	SLATE_BEGIN_ARGS(SWorldPartitionConvertDialog) {}
+		/** A pointer to the parent window */
+		SLATE_ATTRIBUTE(TSharedPtr<SWindow>, ParentWindow)
+		SLATE_ATTRIBUTE(UWorldPartitionConvertOptions*, ConvertOptions)
+	SLATE_END_ARGS()
+
+	~SWorldPartitionConvertDialog() {}
+
+	/**
+	 * Construct this widget.  Called by the SNew() Slate macro.
+	 *
+	 * @param	InArgs				Declaration used by the SNew() macro to construct this widget
+	 * @param	InViewModel			The UI logic not specific to slate
+	 */
+	void Construct(const FArguments& InArgs);
+
+	static const FVector2D DEFAULT_WINDOW_SIZE;
+
+	bool ClickedOk() const { return bClickedOk; }
+	
+private:
+	FReply OnOkClicked();
+	bool IsOkEnabled() const;
+
+	FReply OnCancelClicked();
+		
+	/** Pointer to the parent window, so we know to destroy it when done */
+	TWeakPtr<SWindow> ParentWindowPtr;
+	TWeakObjectPtr<UWorldPartitionConvertOptions> ConvertOptions;
+
+	/** Detailsview used to display SettingsObjects, and allowing user to change options */
+	TSharedPtr<class IDetailsView> DetailsView;
+		
+	/** Dialog Result */
+	bool bClickedOk;
+};
\ No newline at end of file
diff --git a/Engine/Source/Editor/WorldPartitionEditor/Private/WorldPartitionEditorModule.cpp b/Engine/Source/Editor/WorldPartitionEditor/Private/WorldPartitionEditorModule.cpp
--- a/Engine/Source/Editor/WorldPartitionEditor/Private/WorldPartitionEditorModule.cpp
+++ b/Engine/Source/Editor/WorldPartitionEditor/Private/WorldPartitionEditorModule.cpp
@@ -9,9 +9,21 @@
 #include "WorldPartition/SWorldPartitionEditor.h"
 #include "WorldPartition/SWorldPartitionEditorGridSpatialHash.h"
 
+#include "WorldPartition/SWorldPartitionConvertDialog.h"
+#include "WorldPartition/WorldPartitionConvertOptions.h"
+#include "WorldPartition/WorldPartitionEditorSettings.h"
+
 #include "LevelEditor.h"
 #include "Framework/MultiBox/MultiBoxBuilder.h"
 
+#include "Engine/Level.h"
+#include "Misc/MessageDialog.h"
+#include "Misc/StringBuilder.h"
+#include "Misc/ScopedSlowTask.h"
+#include "Interfaces/IMainFrameModule.h"
+#include "Widgets/SWindow.h"
+#include "Commandlets/WorldPartitionConvertCommandlet.h"
+
 IMPLEMENT_MODULE( FWorldPartitionEditorModule, WorldPartitionEditor );
 
 #define LOCTEXT_NAMESPACE "WorldPartition"
@@ -114,4 +126,150 @@ TSharedRef<SWidget> FWorldPartitionEditorModule::CreateWorldPartitionEditor()
 	return SNew(SWorldPartitionEditor).InWorld(EditorWorld);
 }
 
+bool FWorldPartitionEditorModule::ConvertMap(const FString& InLongPackageName)
+{
+	if (!GetDefault<UWorldPartitionEditorSettings>()->bEnableConversionPrompt || ULevel::GetIsLevelPartitionedFromPackage(FName(*InLongPackageName)))
+	{
+		return true;
+	}
+
+	FText Title = LOCTEXT("ConvertMapMsgTitle", "Opening non World Partition map");
+	EAppReturnType::Type Ret = FMessageDialog::Open(EAppMsgType::YesNoCancel, LOCTEXT("ConvertMapMsg", "Do you want to convert this map? (Process will restart editor)"), &Title);
+	if (Ret == EAppReturnType::No)
+	{
+		return true;
+	}
+	else if (Ret == EAppReturnType::Cancel)
+	{
+		return false;
+	}
+
+	UWorldPartitionConvertOptions* DefaultConvertOptions = GetMutableDefault<UWorldPartitionConvertOptions>();
+	DefaultConvertOptions->CommandletClass = GetDefault<UWorldPartitionEditorSettings>()->CommandletClass;
+	DefaultConvertOptions->bInPlace = false;
+	DefaultConvertOptions->bSkipStableGUIDValidation = true;
+	DefaultConvertOptions->LongPackageName = InLongPackageName;
+
+	TSharedPtr<SWindow> DlgWindow =
+		SNew(SWindow)
+		.Title(LOCTEXT("ConvertWindowTitle", "Convert Settings"))
+		.ClientSize(SWorldPartitionConvertDialog::DEFAULT_WINDOW_SIZE)
+		.SizingRule(ESizingRule::UserSized)
+		.SupportsMinimize(false)
+		.SupportsMaximize(false)
+		.SizingRule(ESizingRule::FixedSize);
+
+	TSharedRef<SWorldPartitionConvertDialog> ConvertDialog =
+		SNew(SWorldPartitionConvertDialog)
+		.ParentWindow(DlgWindow)
+		.ConvertOptions(DefaultConvertOptions);
+
+	DlgWindow->SetContent(ConvertDialog);
+
+	IMainFrameModule& MainFrameModule = FModuleManager::LoadModuleChecked<IMainFrameModule>(TEXT("MainFrame"));
+	FSlateApplication::Get().AddModalWindow(DlgWindow.ToSharedRef(), MainFrameModule.GetParentWindow());
+
+	if (ConvertDialog->ClickedOk())
+	{
+		FProcHandle ProcessHandle;
+		bool bCancelled = false;
+
+		// Task scope
+		{
+			FScopedSlowTask SlowTask(0, LOCTEXT("ConvertProgress", "Converting map to world partition..."));
+			SlowTask.MakeDialog(true);
+
+			FString CurrentExecutableName = FPlatformProcess::ExecutablePath();
+			FString ProjectName = FApp::GetProjectName();
+
+			uint32 ProcessID;
+			FString Arguments = FString::Printf(TEXT("%s %s"), *ProjectName, *DefaultConvertOptions->ToCommandletArgs());
+			ProcessHandle = FPlatformProcess::CreateProc(*CurrentExecutableName, *Arguments, true, false, false, &ProcessID, 0, nullptr, nullptr);
+			
+			while (FPlatformProcess::IsProcRunning(ProcessHandle))
+			{
+				if (SlowTask.ShouldCancel())
+				{
+					bCancelled = true;
+					FPlatformProcess::TerminateProc(ProcessHandle);
+					break;
+				}
+
+				SlowTask.EnterProgressFrame(0);
+				FPlatformProcess::Sleep(0.1);
+			}
+		}
+
+		int32 Result = 0;
+		if (!bCancelled && FPlatformProcess::GetProcReturnCode(ProcessHandle, &Result))
+		{	
+			if (Result == 0)
+			{
+				FMessageDialog::Open(EAppMsgType::Ok, LOCTEXT("ConvertMapCompleted", "Conversion succeeded. Editor will now restart."));
+				FUnrealEdMisc::Get().RestartEditor(false);
+			}
+		}
+		else if (bCancelled)
+		{
+			FMessageDialog::Open(EAppMsgType::Ok, LOCTEXT("ConvertMapCancelled", "Conversion cancelled!"));
+		}
+		
+		if(Result != 0)
+		{
+			FMessageDialog::Open(EAppMsgType::Ok, LOCTEXT("ConvertMapFailed", "Conversion failed!"));
+		}
+	}
+
+	return false;
+}
+
+UWorldPartitionEditorSettings::UWorldPartitionEditorSettings()
+{
+	bEnableConversionPrompt = false;
+	CommandletClass = UWorldPartitionConvertCommandlet::StaticClass();
+}
+
+FString UWorldPartitionConvertOptions::ToCommandletArgs() const
+{
+	TStringBuilder<1024> CommandletArgsBuilder;
+	CommandletArgsBuilder.Appendf(TEXT("-run=%s %s -AllowCommandletRendering"), *CommandletClass->GetName(), *LongPackageName);
+	
+	if (!bInPlace)
+	{
+		CommandletArgsBuilder.Append(TEXT(" -ConversionSuffix"));
+	}
+
+	if (bSkipStableGUIDValidation)
+	{
+		CommandletArgsBuilder.Append(TEXT(" -SkipStableGUIDValidation"));
+	}
+
+	if (bSkipMiniMapGeneration)
+	{
+		CommandletArgsBuilder.Append(TEXT(" -SkipMiniMapGeneration"));
+	}
+
+	if (bDeleteSourceLevels)
+	{
+		CommandletArgsBuilder.Append(TEXT(" -DeleteSourceLevels"));
+	}
+	
+	if (bGenerateIni)
+	{
+		CommandletArgsBuilder.Append(TEXT(" -GenerateIni"));
+	}
+	
+	if (bReportOnly)
+	{
+		CommandletArgsBuilder.Append(TEXT(" -ReportOnly"));
+	}
+	
+	if (bVerbose)
+	{
+		CommandletArgsBuilder.Append(TEXT(" -Verbose"));
+	}
+	
+	return CommandletArgsBuilder.ToString();
+}
+
 #undef LOCTEXT_NAMESPACE
diff --git a/Engine/Source/Editor/WorldPartitionEditor/Public/WorldPartition/WorldPartitionConvertOptions.h b/Engine/Source/Editor/WorldPartitionEditor/Public/WorldPartition/WorldPartitionConvertOptions.h
new file mode 100644
--- /dev/null
+++ b/Engine/Source/Editor/WorldPartitionEditor/Public/WorldPartition/WorldPartitionConvertOptions.h
@@ -0,0 +1,44 @@
+// Copyright Epic Games, Inc. All Rights Reserved.
+
+#pragma once
+
+#include "CoreMinimal.h"
+#include "Commandlets/WorldPartitionConvertCommandlet.h"
+#include "WorldPartitionConvertOptions.generated.h"
+
+UCLASS()
+class UWorldPartitionConvertOptions : public UObject
+{
+	GENERATED_BODY()
+
+public:
+
+	UWorldPartitionConvertOptions()
+	{
+	}
+
+	FString ToCommandletArgs() const;
+
+	/** Properties which are supposed to be baked out for the material(s) */
+	UPROPERTY(EditAnywhere, Category=Convert)
+	TSubclassOf<UWorldPartitionConvertCommandlet> CommandletClass;
+		
+	/** Determines whether to not allow usage of the source mesh data while baking out material properties */
+	UPROPERTY(EditAnywhere, Category=Convert, meta = (ToolTip = "Wether the conversion should create a new map with a _WP suffix or overwrite the source map"))
+	bool bInPlace;
+
+	UPROPERTY(EditAnywhere, AdvancedDisplay, Category=Convert)
+	bool bDeleteSourceLevels;
+	UPROPERTY(EditAnywhere, AdvancedDisplay, Category = Convert)
+	bool bGenerateIni;
+	UPROPERTY(EditAnywhere, AdvancedDisplay, Category = Convert)
+	bool bReportOnly;
+	UPROPERTY(EditAnywhere, AdvancedDisplay, Category = Convert)
+	bool bVerbose;
+	UPROPERTY(EditAnywhere, AdvancedDisplay, Category = Convert)
+	bool bSkipStableGUIDValidation;
+	UPROPERTY(EditAnywhere, AdvancedDisplay, Category = Convert)
+	bool bSkipMiniMapGeneration;
+
+	FString LongPackageName;
+};
\ No newline at end of file
diff --git a/Engine/Source/Editor/WorldPartitionEditor/Public/WorldPartition/WorldPartitionEditorSettings.h b/Engine/Source/Editor/WorldPartitionEditor/Public/WorldPartition/WorldPartitionEditorSettings.h
new file mode 100644
--- /dev/null
+++ b/Engine/Source/Editor/WorldPartitionEditor/Public/WorldPartition/WorldPartitionEditorSettings.h
@@ -0,0 +1,24 @@
+// Copyright Epic Games, Inc. All Rights Reserved.
+
+#pragma once
+
+#include "CoreMinimal.h"
+#include "Engine/DeveloperSettings.h"
+#include "Commandlets/WorldPartitionConvertCommandlet.h"
+#include "WorldPartitionEditorSettings.generated.h"
+
+UCLASS(config = Engine, defaultconfig, meta = (DisplayName = "World Partition"))
+class UWorldPartitionEditorSettings : public UDeveloperSettings
+{
+	GENERATED_BODY()
+
+public:
+	UWorldPartitionEditorSettings();
+
+
+	UPROPERTY(config, EditAnywhere, Category = MapConversion, meta = (ToolTip = "This option when enabled will show a conversion prompt when opening non World Partition maps"))
+	bool bEnableConversionPrompt = false;
+
+	UPROPERTY(config, EditAnywhere, Category = MapConversion, meta = (ToolTip = "Commandlet class to use for World Parition conversion"))
+	TSubclassOf<UWorldPartitionConvertCommandlet> CommandletClass;
+};
\ No newline at end of file
diff --git a/Engine/Source/Editor/WorldPartitionEditor/Public/WorldPartitionEditorModule.h b/Engine/Source/Editor/WorldPartitionEditor/Public/WorldPartitionEditorModule.h
--- a/Engine/Source/Editor/WorldPartitionEditor/Public/WorldPartitionEditorModule.h
+++ b/Engine/Source/Editor/WorldPartitionEditor/Public/WorldPartitionEditorModule.h
@@ -2,12 +2,12 @@
 #pragma once
 
 #include "CoreMinimal.h"
-#include "Modules/ModuleInterface.h"
+#include "WorldPartition/IWorldPartitionEditorModule.h"
 
 /**
  * The module holding all of the UI related pieces for SubLevels management
  */
-class FWorldPartitionEditorModule : public IModuleInterface
+class FWorldPartitionEditorModule : public IWorldPartitionEditorModule
 {
 public:
 	/**
@@ -25,6 +25,11 @@ public:
 	 */
 	virtual TSharedRef<class SWidget> CreateWorldPartitionEditor();
 	
+	/**
+	 * 
+	 */
+	virtual bool ConvertMap(const FString& InLongPackageName);
+
 private:
 	FDelegateHandle LevelEditorExtenderDelegateHandle;
 
diff --git a/Engine/Source/Editor/WorldPartitionEditor/WorldPartitionEditor.Build.cs b/Engine/Source/Editor/WorldPartitionEditor/WorldPartitionEditor.Build.cs
--- a/Engine/Source/Editor/WorldPartitionEditor/WorldPartitionEditor.Build.cs
+++ b/Engine/Source/Editor/WorldPartitionEditor/WorldPartitionEditor.Build.cs
@@ -27,6 +27,9 @@ public class WorldPartitionEditor : ModuleRules
 				"SlateCore",
 				"UnrealEd",
 				"WorldBrowser",
+				"MainFrame",
+				"PropertyEditor",
+				"DeveloperSettings"
             }
 		);
 
diff --git a/Engine/Source/Runtime/Engine/Public/WorldPartition/IWorldPartitionEditorModule.h b/Engine/Source/Runtime/Engine/Public/WorldPartition/IWorldPartitionEditorModule.h
new file mode 100644
--- /dev/null
+++ b/Engine/Source/Runtime/Engine/Public/WorldPartition/IWorldPartitionEditorModule.h
@@ -0,0 +1,16 @@
+// Copyright Epic Games, Inc. All Rights Reserved.
+#pragma once
+
+#include "CoreMinimal.h"
+#include "Modules/ModuleInterface.h"
+
+/**
+ * The module holding all of the UI related pieces for WorldPartition
+ */
+class IWorldPartitionEditorModule : public IModuleInterface
+{
+public:
+	virtual ~IWorldPartitionEditorModule() {}
+
+	virtual bool ConvertMap(const FString& InLongPackageName) = 0;
+};
\ No newline at end of file
