World Partition Editor
- Auto-focus on world bounds when opening window.
- Added button to focus on selected actors.

diff --git a/Engine/Source/Editor/WorldPartitionEditor/Private/WorldPartition/SWorldPartitionEditor.cpp b/Engine/Source/Editor/WorldPartitionEditor/Private/WorldPartition/SWorldPartitionEditor.cpp
--- a/Engine/Source/Editor/WorldPartitionEditor/Private/WorldPartition/SWorldPartitionEditor.cpp
+++ b/Engine/Source/Editor/WorldPartitionEditor/Private/WorldPartition/SWorldPartitionEditor.cpp
@@ -40,16 +40,6 @@ SWorldPartitionEditor::~SWorldPartitionEditor()
 	}
 }
 
-void SWorldPartitionEditor::InvalidatePartition()
-{
-	GridView.Get()->InvalidatePartition();
-}
-
-void SWorldPartitionEditor::RecreatePartition()
-{
-	ContentParent->SetContent(ConstructContentWidget());
-}
-
 void SWorldPartitionEditor::Refresh()
 {
 	GridView->RefreshSceneOutliner();
@@ -58,7 +48,7 @@ void SWorldPartitionEditor::Refresh()
 void SWorldPartitionEditor::OnBrowseWorld(UWorld* InWorld)
 {
 	World = InWorld;
-	RecreatePartition();
+	ContentParent->SetContent(ConstructContentWidget());
 
 	// No need to unregister as the previous UWorldPartitionSubsystem is already destroyed
 	WorldPartitionChangedDelegateHandle.Reset();
diff --git a/Engine/Source/Editor/WorldPartitionEditor/Private/WorldPartition/SWorldPartitionEditorGrid2D.cpp b/Engine/Source/Editor/WorldPartitionEditor/Private/WorldPartition/SWorldPartitionEditorGrid2D.cpp
--- a/Engine/Source/Editor/WorldPartitionEditor/Private/WorldPartition/SWorldPartitionEditorGrid2D.cpp
+++ b/Engine/Source/Editor/WorldPartitionEditor/Private/WorldPartition/SWorldPartitionEditorGrid2D.cpp
@@ -21,7 +21,7 @@
 
 #define LOCTEXT_NAMESPACE "WorldPartitionEditor"
 
-static TAutoConsoleVariable<int32> CVarShowReloadMiniMapButton(TEXT("MiniMap.ShowReloadButton"), 0, TEXT("Show reload MiniMap button."));
+static TAutoConsoleVariable<int32> CVarShowReloadMiniMapButton(TEXT("wp.MiniMap.ShowReloadButton"), 0, TEXT("Show reload MiniMap button."));
 
 SWorldPartitionEditorGrid2D::FEditorCommands::FEditorCommands()
 	: TCommands<FEditorCommands>
@@ -108,10 +108,17 @@ void SWorldPartitionEditorGrid2D::Construct(const FArguments& InArgs)
 					.AutoWidth()
 					[
 						SNew(SButton)
-						.Text(LOCTEXT("Reload MiniMap", "Reload MiniMap "))
+						.Text(LOCTEXT("ReloadMiniMap", "Reload MiniMap"))
 						.Visibility_Lambda([]() {return CVarShowReloadMiniMapButton.GetValueOnAnyThread() != 0 ? EVisibility::Visible : EVisibility::Hidden;})
 						.OnClicked(this, &SWorldPartitionEditorGrid2D::ReloadMiniMap)
 					]
+					+ SHorizontalBox::Slot()
+					.AutoWidth()
+					[
+						SNew(SButton)
+						.Text(LOCTEXT("FocusSelection", "Focus Selection"))
+						.OnClicked(this, &SWorldPartitionEditorGrid2D::FocusSelection)
+					]
 				]
 			]
 		]
@@ -158,15 +165,6 @@ void SWorldPartitionEditorGrid2D::MoveCameraHere()
 	}
 }
 
-void SWorldPartitionEditorGrid2D::InvalidatePartition()
-{
-	SWorldPartitionEditorGrid::InvalidatePartition();
-	
-	SelectionStart = FVector2D::ZeroVector;
-	SelectionEnd = FVector2D::ZeroVector;
-	UpdateSelection();
-}
-
 FReply SWorldPartitionEditorGrid2D::OnMouseButtonDown(const FGeometry& MyGeometry, const FPointerEvent& MouseEvent)
 {
 	const bool bIsLeftMouseButtonEffecting = MouseEvent.GetEffectingButton() == EKeys::LeftMouseButton;
@@ -609,7 +607,16 @@ int32 SWorldPartitionEditorGrid2D::OnPaint(const FPaintArgs& Args, const FGeomet
 {
 	if (WorldPartition)
 	{
+		const bool bResetView = !ScreenRect.bIsValid;
+
 		ScreenRect = FBox2D(FVector2D(0, 0), AllottedGeometry.GetLocalSize());
+
+		if (bResetView)
+		{
+			const FBox WorldBounds = WorldPartition->GetWorldBounds();
+			FocusBox(WorldBounds);
+		}
+
 		UpdateTransform();
 
 		LayerId = PaintGrid(AllottedGeometry, MyCullingRect, OutDrawElements, ++LayerId);
@@ -651,6 +658,43 @@ int32 SWorldPartitionEditorGrid2D::PaintSoftwareCursor(const FGeometry& Allotted
 	return LayerId + 1;
 }
 
+FReply SWorldPartitionEditorGrid2D::FocusSelection()
+{
+	FBox SelectionBox(ForceInit);
+
+	USelection* SelectedActors = GEditor->GetSelectedActors();
+	for (FSelectionIterator It(*SelectedActors); It; ++It)
+	{
+		if (AActor* Actor = Cast<AActor>(*It))
+		{
+			FVector Origin;
+			FVector Extent;
+			Actor->GetActorBounds(false, Origin, Extent);
+
+			SelectionBox += FBox(Origin - Extent, Origin + Extent);
+		}
+	}
+
+	FocusBox(SelectionBox);
+	return FReply::Handled();
+}
+
+void SWorldPartitionEditorGrid2D::FocusBox(const FBox& Box) const
+{
+	check(ScreenRect.bIsValid);
+
+	if (Box.GetVolume() > 0.0f)
+	{
+		const FVector2D ScreenExtent = ScreenRect.GetExtent();
+		const FVector2D SelectExtent = FVector2D(Box.GetExtent());
+
+		Trans = -FVector2D(Box.GetCenter());
+		Scale = (ScreenExtent / SelectExtent).GetMin() * 0.75f;
+
+		UpdateTransform();
+	}	
+}
+
 void SWorldPartitionEditorGrid2D::UpdateTransform() const
 {
 	FTransform2D T(1.0f, Trans);
diff --git a/Engine/Source/Editor/WorldPartitionEditor/Private/WorldPartition/SWorldPartitionEditorGridSpatialHash.cpp b/Engine/Source/Editor/WorldPartitionEditor/Private/WorldPartition/SWorldPartitionEditorGridSpatialHash.cpp
--- a/Engine/Source/Editor/WorldPartitionEditor/Private/WorldPartition/SWorldPartitionEditorGridSpatialHash.cpp
+++ b/Engine/Source/Editor/WorldPartitionEditor/Private/WorldPartition/SWorldPartitionEditorGridSpatialHash.cpp
@@ -20,6 +20,7 @@
 #include "ObjectTools.h"
 #include "UObject/ObjectResource.h"
 #include "EditorStyleSet.h"
+#include "Selection.h"
 #include "Engine/World.h"
 #include "Engine/Texture2DDynamic.h"
 #include "Misc/HashBuilder.h"
@@ -55,6 +56,18 @@ void SWorldPartitionEditorGridSpatialHash::Construct(const FArguments& InArgs)
 	SWorldPartitionEditorGrid2D::Construct(SWorldPartitionEditorGrid::FArguments().InWorld(InArgs._InWorld));
 }
 
+void SWorldPartitionEditorGridSpatialHash::Tick(const FGeometry& AllottedGeometry, const double InCurrentTime, const float InDeltaTime)
+{
+	SWorldPartitionEditorGrid2D::Tick(AllottedGeometry, InCurrentTime, InDeltaTime);
+
+	/*static double LastTime = InCurrentTime;
+	if ((InCurrentTime - LastTime) > 15.0)
+	{
+		LastTime = InCurrentTime;
+		ReloadMiniMap();
+	}*/
+}
+
 FReply SWorldPartitionEditorGridSpatialHash::ReloadMiniMap()
 {
 	UE_LOG(LogTemp, Log, TEXT("Reload MiniMap has been clicked"));
@@ -87,7 +100,6 @@ void SWorldPartitionEditorGridSpatialHash::UpdateWorldMiniMapDetails()
 			WorldMiniMapBrush.SetResourceObject(MiniMapTexture);
 		}
 	}
-	
 }
 
 int32 SWorldPartitionEditorGridSpatialHash::PaintGrid(const FGeometry& AllottedGeometry, const FSlateRect& MyCullingRect, FSlateWindowElementList& OutDrawElements, int32 LayerId) const
diff --git a/Engine/Source/Editor/WorldPartitionEditor/Public/WorldPartition/SWorldPartitionEditor.h b/Engine/Source/Editor/WorldPartitionEditor/Public/WorldPartition/SWorldPartitionEditor.h
--- a/Engine/Source/Editor/WorldPartitionEditor/Public/WorldPartition/SWorldPartitionEditor.h
+++ b/Engine/Source/Editor/WorldPartitionEditor/Public/WorldPartition/SWorldPartitionEditor.h
@@ -21,8 +21,6 @@ public:
 	~SWorldPartitionEditor();
 
 	// IWorldPartitionEditor interface
-	virtual void InvalidatePartition() override;
-	virtual void RecreatePartition() override;
 	virtual void Refresh() override;
 
 private:
diff --git a/Engine/Source/Editor/WorldPartitionEditor/Public/WorldPartition/SWorldPartitionEditorGrid.h b/Engine/Source/Editor/WorldPartitionEditor/Public/WorldPartition/SWorldPartitionEditorGrid.h
--- a/Engine/Source/Editor/WorldPartitionEditor/Public/WorldPartition/SWorldPartitionEditorGrid.h
+++ b/Engine/Source/Editor/WorldPartitionEditor/Public/WorldPartition/SWorldPartitionEditorGrid.h
@@ -57,8 +57,6 @@ public:
 		return &NoChildrenInstance;
 	}
 	
-	virtual void InvalidatePartition() {}
-
 	bool GetPlayerView(FVector& Location, FRotator& Rotation) const;
 	bool GetObserverView(FVector& Location, FRotator& Rotation) const;
 
diff --git a/Engine/Source/Editor/WorldPartitionEditor/Public/WorldPartition/SWorldPartitionEditorGrid2D.h b/Engine/Source/Editor/WorldPartitionEditor/Public/WorldPartition/SWorldPartitionEditorGrid2D.h
--- a/Engine/Source/Editor/WorldPartitionEditor/Public/WorldPartition/SWorldPartitionEditorGrid2D.h
+++ b/Engine/Source/Editor/WorldPartitionEditor/Public/WorldPartition/SWorldPartitionEditorGrid2D.h
@@ -38,9 +38,6 @@ public:
 	void UnloadSelectedCells();
 	void MoveCameraHere();
 
-	// IWorldPartitionEditor interface
-	virtual void InvalidatePartition() override;
-
 	virtual FReply OnMouseButtonDown(const FGeometry& MyGeometry, const FPointerEvent& MouseEvent) override;
 	virtual FReply OnMouseButtonUp(const FGeometry& MyGeometry, const FPointerEvent& MouseEvent) override;
 	virtual FReply OnMouseMove(const FGeometry& MyGeometry, const FPointerEvent& MouseEvent) override;
@@ -56,6 +53,7 @@ public:
 	virtual int32 PaintSoftwareCursor(const FGeometry& AllottedGeometry, const FSlateRect& MyCullingRect, FSlateWindowElementList& OutDrawElements, int32 LayerId) const;
 
 	virtual FReply ReloadMiniMap() { return FReply::Handled(); }
+	virtual FReply FocusSelection();
 
 protected:
 	void UpdateTransform() const;
@@ -75,8 +73,10 @@ protected:
 		ArrangedChildren.AddWidget(AllottedGeometry.MakeChild(ChildSlot.GetWidget(), FVector2D(0, 0), AllottedGeometry.GetAbsoluteSize())); 
 	}
 
-	float Scale;
-	FVector2D Trans;
+	void FocusBox(const FBox& Box) const;
+
+	mutable float Scale;
+	mutable FVector2D Trans;
 	
 	mutable FBox2D ScreenRect;
 	mutable FTransform2D WorldToScreen;
diff --git a/Engine/Source/Editor/WorldPartitionEditor/Public/WorldPartition/SWorldPartitionEditorGridSpatialHash.h b/Engine/Source/Editor/WorldPartitionEditor/Public/WorldPartition/SWorldPartitionEditorGridSpatialHash.h
--- a/Engine/Source/Editor/WorldPartitionEditor/Public/WorldPartition/SWorldPartitionEditorGridSpatialHash.h
+++ b/Engine/Source/Editor/WorldPartitionEditor/Public/WorldPartition/SWorldPartitionEditorGridSpatialHash.h
@@ -15,10 +15,12 @@ public:
 	~SWorldPartitionEditorGridSpatialHash();
 
 	void Construct(const FArguments& InArgs);
+
+	virtual void Tick(const FGeometry& AllottedGeometry, const double InCurrentTime, const float InDeltaTime) override;
 	
 	int32 PaintGrid(const FGeometry& AllottedGeometry, const FSlateRect& MyCullingRect, FSlateWindowElementList& OutDrawElements, int32 LayerId) const override;
 
-	FReply ReloadMiniMap() override;
+	virtual FReply ReloadMiniMap() override;
 
 private:
 	void UpdateWorldMiniMapDetails();
diff --git a/Engine/Source/Runtime/Engine/Public/WorldPartition/WorldPartition.h b/Engine/Source/Runtime/Engine/Public/WorldPartition/WorldPartition.h
--- a/Engine/Source/Runtime/Engine/Public/WorldPartition/WorldPartition.h
+++ b/Engine/Source/Runtime/Engine/Public/WorldPartition/WorldPartition.h
@@ -44,8 +44,6 @@ enum class EWorldPartitionInitState
  */
 struct ENGINE_API IWorldPartitionEditor
 {
-	virtual void InvalidatePartition() {}
-	virtual void RecreatePartition() {}
 	virtual void Refresh() {}
 };
 
