Fix Non-unity build

diff --git a/Engine/Source/Editor/WorldPartitionEditor/Private/WorldPartition/SWorldPartitionConvertDialog.cpp b/Engine/Source/Editor/WorldPartitionEditor/Private/WorldPartition/SWorldPartitionConvertDialog.cpp
--- a/Engine/Source/Editor/WorldPartitionEditor/Private/WorldPartition/SWorldPartitionConvertDialog.cpp
+++ b/Engine/Source/Editor/WorldPartitionEditor/Private/WorldPartition/SWorldPartitionConvertDialog.cpp
@@ -1,7 +1,10 @@
 // Copyright Epic Games, Inc. All Rights Reserved.
 #include "SWorldPartitionConvertDialog.h"
 #include "Widgets/Layout/SSpacer.h"
+#include "Widgets/Input/SButton.h"
 #include "WorldPartition/WorldPartitionConvertOptions.h"
+#include "Modules/ModuleManager.h"
+
 
 #define LOCTEXT_NAMESPACE "WorldPartitionConvertDialog"
 
