Don't show actors in the world partition editor when a minimap is present.

diff --git a/Engine/Source/Editor/WorldPartitionEditor/Private/WorldPartition/SWorldPartitionEditorGridSpatialHash.cpp b/Engine/Source/Editor/WorldPartitionEditor/Private/WorldPartition/SWorldPartitionEditorGridSpatialHash.cpp
--- a/Engine/Source/Editor/WorldPartitionEditor/Private/WorldPartition/SWorldPartitionEditorGridSpatialHash.cpp
+++ b/Engine/Source/Editor/WorldPartitionEditor/Private/WorldPartition/SWorldPartitionEditorGridSpatialHash.cpp
@@ -47,10 +47,10 @@ void SWorldPartitionEditorGridSpatialHash::Construct(const FArguments& InArgs)
 	{
 		UWorldPartitionEditorSpatialHash* EditorSpatialHash = (UWorldPartitionEditorSpatialHash*)WorldPartition->EditorHash;
 
-		bShowActors = true;
-
 		//Update MiniMap data for drawing  
 		UpdateWorldMiniMapDetails();
+
+		bShowActors = !WorldMiniMapBrush.HasUObject();
 	}
 
 	SWorldPartitionEditorGrid2D::Construct(SWorldPartitionEditorGrid::FArguments().InWorld(InArgs._InWorld));
