WorldPartition: Pass in full project path to external conversion process (fixes failed conversion on non-rooted projects)

diff --git a/Engine/Source/Editor/WorldPartitionEditor/Private/WorldPartitionEditorModule.cpp b/Engine/Source/Editor/WorldPartitionEditor/Private/WorldPartitionEditorModule.cpp
--- a/Engine/Source/Editor/WorldPartitionEditor/Private/WorldPartitionEditorModule.cpp
+++ b/Engine/Source/Editor/WorldPartitionEditor/Private/WorldPartitionEditorModule.cpp
@@ -200,10 +200,12 @@ bool FWorldPartitionEditorModule::ConvertMap(const FString& InLongPackageName)
 			SlowTask.MakeDialog(true);
 
 			FString CurrentExecutableName = FPlatformProcess::ExecutablePath();
-			FString ProjectName = FApp::GetProjectName();
 
+			// Try to provide complete Path, if we can't try with project name
+			FString ProjectPath = FPaths::IsProjectFilePathSet() ? FPaths::GetProjectFilePath() : FApp::GetProjectName();
+			
 			uint32 ProcessID;
-			FString Arguments = FString::Printf(TEXT("%s %s"), *ProjectName, *DefaultConvertOptions->ToCommandletArgs());
+			FString Arguments = FString::Printf(TEXT("\"%s\" %s"), *ProjectPath, *DefaultConvertOptions->ToCommandletArgs());
 			ProcessHandle = FPlatformProcess::CreateProc(*CurrentExecutableName, *Arguments, true, false, false, &ProcessID, 0, nullptr, nullptr);
 			
 			while (FPlatformProcess::IsProcRunning(ProcessHandle))
