WorldPartition: Load Blank map before conversion

diff --git a/Engine/Source/Editor/WorldPartitionEditor/Private/WorldPartitionEditorModule.cpp b/Engine/Source/Editor/WorldPartitionEditor/Private/WorldPartitionEditorModule.cpp
--- a/Engine/Source/Editor/WorldPartitionEditor/Private/WorldPartitionEditorModule.cpp
+++ b/Engine/Source/Editor/WorldPartitionEditor/Private/WorldPartitionEditorModule.cpp
@@ -23,6 +23,7 @@
 #include "Interfaces/IMainFrameModule.h"
 #include "Widgets/SWindow.h"
 #include "Commandlets/WorldPartitionConvertCommandlet.h"
+#include "FileHelpers.h"
 
 IMPLEMENT_MODULE( FWorldPartitionEditorModule, WorldPartitionEditor );
 
@@ -191,6 +192,12 @@ bool FWorldPartitionEditorModule::ConvertMap(const FString& InLongPackageName)
 
 	if (ConvertDialog->ClickedOk())
 	{
+		// User will already have been prompted to save on file open
+		if (!UEditorLoadingAndSavingUtils::NewBlankMap(/*bSaveExistingMap*/false))
+		{
+			return false;
+		}
+
 		FProcHandle ProcessHandle;
 		bool bCancelled = false;
 
