Don't ask to convert non-partitioned worlds when world partition is not enabled in the project settings.

diff --git a/Engine/Source/Editor/WorldPartitionEditor/Private/WorldPartitionEditorModule.cpp b/Engine/Source/Editor/WorldPartitionEditor/Private/WorldPartitionEditorModule.cpp
--- a/Engine/Source/Editor/WorldPartitionEditor/Private/WorldPartitionEditorModule.cpp
+++ b/Engine/Source/Editor/WorldPartitionEditor/Private/WorldPartitionEditorModule.cpp
@@ -127,14 +127,14 @@ TSharedRef<SWidget> FWorldPartitionEditorModule::CreateWorldPartitionEditor()
 	return SNew(SWorldPartitionEditor).InWorld(EditorWorld);
 }
 
-bool FWorldPartitionEditorModule::IsWorldPartitionEnabled()
+bool FWorldPartitionEditorModule::IsWorldPartitionEnabled() const
 {
 	return GetDefault<UWorldPartitionEditorSettings>()->bEnableWorldPartition;
 }
 
 bool FWorldPartitionEditorModule::IsConversionPromptEnabled() const
 {
-	return GetDefault<UWorldPartitionEditorSettings>()->bEnableConversionPrompt;
+	return IsWorldPartitionEnabled() && GetDefault<UWorldPartitionEditorSettings>()->bEnableConversionPrompt;
 }
 
 void FWorldPartitionEditorModule::SetConversionPromptEnabled(bool bEnabled)
@@ -149,7 +149,7 @@ float FWorldPartitionEditorModule::GetAutoCellLoadingMaxWorldSize() const
 
 bool FWorldPartitionEditorModule::ConvertMap(const FString& InLongPackageName)
 {
-	if (!GetDefault<UWorldPartitionEditorSettings>()->bEnableConversionPrompt || ULevel::GetIsLevelPartitionedFromPackage(FName(*InLongPackageName)))
+	if (!IsConversionPromptEnabled() || ULevel::GetIsLevelPartitionedFromPackage(FName(*InLongPackageName)))
 	{
 		return true;
 	}
diff --git a/Engine/Source/Editor/WorldPartitionEditor/Public/WorldPartitionEditorModule.h b/Engine/Source/Editor/WorldPartitionEditor/Public/WorldPartitionEditorModule.h
--- a/Engine/Source/Editor/WorldPartitionEditor/Public/WorldPartitionEditorModule.h
+++ b/Engine/Source/Editor/WorldPartitionEditor/Public/WorldPartitionEditorModule.h
@@ -28,7 +28,7 @@ public:
 	/**
 	 *
 	 */
-	virtual bool IsWorldPartitionEnabled() override;
+	virtual bool IsWorldPartitionEnabled() const override;
 
 	/**
 	 *
diff --git a/Engine/Source/Runtime/Engine/Public/WorldPartition/IWorldPartitionEditorModule.h b/Engine/Source/Runtime/Engine/Public/WorldPartition/IWorldPartitionEditorModule.h
--- a/Engine/Source/Runtime/Engine/Public/WorldPartition/IWorldPartitionEditorModule.h
+++ b/Engine/Source/Runtime/Engine/Public/WorldPartition/IWorldPartitionEditorModule.h
@@ -12,7 +12,7 @@ class IWorldPartitionEditorModule : public IModuleInterface
 public:
 	virtual ~IWorldPartitionEditorModule() {}
 
-	virtual bool IsWorldPartitionEnabled() = 0;
+	virtual bool IsWorldPartitionEnabled() const = 0;
 	virtual bool IsConversionPromptEnabled() const = 0;
 	virtual void SetConversionPromptEnabled(bool bEnabled) = 0;
 	virtual float GetAutoCellLoadingMaxWorldSize() const = 0;
