Cache the world bounds in World Partition Editor instead of computing it every tick, to make sure it scales for huge worlds.

diff --git a/Engine/Source/Editor/WorldPartitionEditor/Private/WorldPartition/SWorldPartitionEditor.cpp b/Engine/Source/Editor/WorldPartitionEditor/Private/WorldPartition/SWorldPartitionEditor.cpp
--- a/Engine/Source/Editor/WorldPartitionEditor/Private/WorldPartition/SWorldPartitionEditor.cpp
+++ b/Engine/Source/Editor/WorldPartitionEditor/Private/WorldPartition/SWorldPartitionEditor.cpp
@@ -42,7 +42,7 @@ SWorldPartitionEditor::~SWorldPartitionEditor()
 
 void SWorldPartitionEditor::Refresh()
 {
-	GridView->RefreshSceneOutliner();
+	GridView->Refresh();
 }
 
 void SWorldPartitionEditor::OnBrowseWorld(UWorld* InWorld)
@@ -118,6 +118,7 @@ TSharedRef<SWidget> SWorldPartitionEditor::ConstructContentWidget()
 			]
 		];
 
+	Refresh();
 	return Result;
 }
 
diff --git a/Engine/Source/Editor/WorldPartitionEditor/Private/WorldPartition/SWorldPartitionEditorGrid.cpp b/Engine/Source/Editor/WorldPartitionEditor/Private/WorldPartition/SWorldPartitionEditorGrid.cpp
--- a/Engine/Source/Editor/WorldPartitionEditor/Private/WorldPartition/SWorldPartitionEditorGrid.cpp
+++ b/Engine/Source/Editor/WorldPartitionEditor/Private/WorldPartition/SWorldPartitionEditorGrid.cpp
@@ -74,7 +74,7 @@ bool SWorldPartitionEditorGrid::GetObserverView(FVector& Location, FRotator& Rot
 	return false;
 }
 
-void SWorldPartitionEditorGrid::RefreshSceneOutliner()
+void SWorldPartitionEditorGrid::Refresh()
 {
 	TWeakPtr<class ILevelEditor> LevelEditor = FModuleManager::GetModuleChecked<FLevelEditorModule>(TEXT("LevelEditor")).GetLevelEditorInstance();
 	if (LevelEditor.IsValid())
@@ -85,4 +85,10 @@ void SWorldPartitionEditorGrid::RefreshSceneOutliner()
 			SceneOutlinerPtr->FullRefresh();
 		}
 	}
+
+	WorldBounds.Init();
+	if (WorldPartition)
+	{
+		WorldBounds = WorldPartition->GetWorldBounds();
+	}
 }
\ No newline at end of file
diff --git a/Engine/Source/Editor/WorldPartitionEditor/Private/WorldPartition/SWorldPartitionEditorGrid2D.cpp b/Engine/Source/Editor/WorldPartitionEditor/Private/WorldPartition/SWorldPartitionEditorGrid2D.cpp
--- a/Engine/Source/Editor/WorldPartitionEditor/Private/WorldPartition/SWorldPartitionEditorGrid2D.cpp
+++ b/Engine/Source/Editor/WorldPartitionEditor/Private/WorldPartition/SWorldPartitionEditorGrid2D.cpp
@@ -144,14 +144,14 @@ void SWorldPartitionEditorGrid2D::LoadSelectedCells()
 {
 	WorldPartition->LoadEditorCells(SelectBox);
 	GEditor->RedrawLevelEditingViewports();
-	RefreshSceneOutliner();
+	Refresh();
 }
 
 void SWorldPartitionEditorGrid2D::UnloadSelectedCells()
 {
 	WorldPartition->UnloadEditorCells(SelectBox);
 	GEditor->RedrawLevelEditingViewports();
-	RefreshSceneOutliner();
+	Refresh();
 }
 
 void SWorldPartitionEditorGrid2D::MoveCameraHere()
@@ -480,7 +480,6 @@ uint32 SWorldPartitionEditorGrid2D::PaintScaleRuler(const FGeometry& AllottedGeo
 		FLinearColor::White);
 
 	// Show world bounds
-	const FBox WorldBounds = WorldPartition->GetWorldBounds();
 	const FVector WorldBoundsExtentInKM = (WorldBounds.GetExtent() * 2.0f) / 100000.0f;
 	RulerText = FString::Printf(TEXT("%.2fx%.2fx%.2f km"), WorldBoundsExtentInKM.X, WorldBoundsExtentInKM.Y, WorldBoundsExtentInKM.Z);
 	
@@ -609,7 +608,6 @@ int32 SWorldPartitionEditorGrid2D::OnPaint(const FPaintArgs& Args, const FGeomet
 
 		if (bResetView)
 		{
-			const FBox WorldBounds = WorldPartition->GetWorldBounds();
 			FocusBox(WorldBounds);
 		}
 
diff --git a/Engine/Source/Editor/WorldPartitionEditor/Public/WorldPartition/SWorldPartitionEditorGrid.h b/Engine/Source/Editor/WorldPartitionEditor/Public/WorldPartition/SWorldPartitionEditorGrid.h
--- a/Engine/Source/Editor/WorldPartitionEditor/Public/WorldPartition/SWorldPartitionEditorGrid.h
+++ b/Engine/Source/Editor/WorldPartitionEditor/Public/WorldPartition/SWorldPartitionEditorGrid.h
@@ -60,10 +60,11 @@ public:
 	bool GetPlayerView(FVector& Location, FRotator& Rotation) const;
 	bool GetObserverView(FVector& Location, FRotator& Rotation) const;
 
-	void RefreshSceneOutliner();
+	void Refresh();
 
 	UWorld* World;
 	UWorldPartition* WorldPartition;
+	FBox WorldBounds;
 
 protected:
 	static TMap<FName, PartitionEditorGridCreateInstanceFunc> PartitionEditorGridCreateInstanceFactory;
diff --git a/Engine/Source/Runtime/Engine/Private/WorldPartition/WorldPartition.cpp b/Engine/Source/Runtime/Engine/Private/WorldPartition/WorldPartition.cpp
--- a/Engine/Source/Runtime/Engine/Private/WorldPartition/WorldPartition.cpp
+++ b/Engine/Source/Runtime/Engine/Private/WorldPartition/WorldPartition.cpp
@@ -703,6 +703,11 @@ void UWorldPartition::OnAssetAdded(const FAssetData& InAssetData)
 		HashActorDesc(NewActorDesc->Get());
 
 		(*NewActorDesc)->OnRegister();
+
+		if (WorldPartitionEditor)
+		{
+			WorldPartitionEditor->Refresh();
+		}
 	}
 }
 
@@ -721,6 +726,11 @@ void UWorldPartition::OnAssetRemoved(const FAssetData& InAssetData)
 
 		Actors.Remove((*ExistingActorDesc)->GetGuid());
 		ExistingActorDesc->Release();
+
+		if (WorldPartitionEditor)
+		{
+			WorldPartitionEditor->Refresh();
+		}
 	}
 }
 
@@ -745,6 +755,11 @@ void UWorldPartition::OnAssetUpdated(const FAssetData& InAssetData)
 		*ExistingActorDesc = MoveTemp(NewActorDesc);
 		
 		HashActorDesc(ExistingActorDesc->Get());
+
+		if (WorldPartitionEditor)
+		{
+			WorldPartitionEditor->Refresh();
+		}
 	}
 }
 
