Use editor world bounds instead of runtime world bounds for the World Partition editor.
Also focus on the whole world when no actors are selected for the "focus selection" button.

diff --git a/Engine/Source/Editor/WorldPartitionEditor/Private/WorldPartition/SWorldPartitionEditorGrid.cpp b/Engine/Source/Editor/WorldPartitionEditor/Private/WorldPartition/SWorldPartitionEditorGrid.cpp
--- a/Engine/Source/Editor/WorldPartitionEditor/Private/WorldPartition/SWorldPartitionEditorGrid.cpp
+++ b/Engine/Source/Editor/WorldPartitionEditor/Private/WorldPartition/SWorldPartitionEditorGrid.cpp
@@ -85,10 +85,4 @@ void SWorldPartitionEditorGrid::Refresh()
 			SceneOutlinerPtr->FullRefresh();
 		}
 	}
-
-	WorldBounds.Init();
-	if (WorldPartition)
-	{
-		WorldBounds = WorldPartition->GetWorldBounds();
-	}
 }
\ No newline at end of file
diff --git a/Engine/Source/Editor/WorldPartitionEditor/Private/WorldPartition/SWorldPartitionEditorGrid2D.cpp b/Engine/Source/Editor/WorldPartitionEditor/Private/WorldPartition/SWorldPartitionEditorGrid2D.cpp
--- a/Engine/Source/Editor/WorldPartitionEditor/Private/WorldPartition/SWorldPartitionEditorGrid2D.cpp
+++ b/Engine/Source/Editor/WorldPartitionEditor/Private/WorldPartition/SWorldPartitionEditorGrid2D.cpp
@@ -480,6 +480,7 @@ uint32 SWorldPartitionEditorGrid2D::PaintScaleRuler(const FGeometry& AllottedGeo
 		FLinearColor::White);
 
 	// Show world bounds
+	const FBox WorldBounds = WorldPartition->GetEditorWorldBounds();
 	const FVector WorldBoundsExtentInKM = (WorldBounds.GetExtent() * 2.0f) / 100000.0f;
 	RulerText = FString::Printf(TEXT("%.2fx%.2fx%.2f km"), WorldBoundsExtentInKM.X, WorldBoundsExtentInKM.Y, WorldBoundsExtentInKM.Z);
 	
@@ -608,6 +609,7 @@ int32 SWorldPartitionEditorGrid2D::OnPaint(const FPaintArgs& Args, const FGeomet
 
 		if (bResetView)
 		{
+			const FBox WorldBounds = WorldPartition->GetEditorWorldBounds();
 			FocusBox(WorldBounds);
 		}
 
@@ -657,17 +659,25 @@ FReply SWorldPartitionEditorGrid2D::FocusSelection()
 	FBox SelectionBox(ForceInit);
 
 	USelection* SelectedActors = GEditor->GetSelectedActors();
-	for (FSelectionIterator It(*SelectedActors); It; ++It)
+
+	if (SelectedActors->Num())
 	{
-		if (AActor* Actor = Cast<AActor>(*It))
+		for (FSelectionIterator It(*SelectedActors); It; ++It)
 		{
-			FVector Origin;
-			FVector Extent;
-			Actor->GetActorBounds(false, Origin, Extent);
+			if (AActor* Actor = Cast<AActor>(*It))
+			{
+				FVector Origin;
+				FVector Extent;
+				Actor->GetActorBounds(false, Origin, Extent);
 
-			SelectionBox += FBox(Origin - Extent, Origin + Extent);
+				SelectionBox += FBox(Origin - Extent, Origin + Extent);
+			}
 		}
 	}
+	else
+	{
+		SelectionBox = WorldPartition->GetEditorWorldBounds();
+	}
 
 	FocusBox(SelectionBox);
 	return FReply::Handled();
diff --git a/Engine/Source/Editor/WorldPartitionEditor/Public/WorldPartition/SWorldPartitionEditorGrid.h b/Engine/Source/Editor/WorldPartitionEditor/Public/WorldPartition/SWorldPartitionEditorGrid.h
--- a/Engine/Source/Editor/WorldPartitionEditor/Public/WorldPartition/SWorldPartitionEditorGrid.h
+++ b/Engine/Source/Editor/WorldPartitionEditor/Public/WorldPartition/SWorldPartitionEditorGrid.h
@@ -64,7 +64,6 @@ public:
 
 	UWorld* World;
 	UWorldPartition* WorldPartition;
-	FBox WorldBounds;
 
 protected:
 	static TMap<FName, PartitionEditorGridCreateInstanceFunc> PartitionEditorGridCreateInstanceFactory;
diff --git a/Engine/Source/Runtime/Engine/Private/WorldPartition/WorldPartition.cpp b/Engine/Source/Runtime/Engine/Private/WorldPartition/WorldPartition.cpp
--- a/Engine/Source/Runtime/Engine/Private/WorldPartition/WorldPartition.cpp
+++ b/Engine/Source/Runtime/Engine/Private/WorldPartition/WorldPartition.cpp
@@ -265,7 +265,7 @@ void UWorldPartition::Initialize(UWorld* InWorld, const FTransform& InTransform)
 				// Autoload all cells if the world is smaller than the project setting's value
 				IWorldPartitionEditorModule& WorldPartitionEditorModule = FModuleManager::LoadModuleChecked<IWorldPartitionEditorModule>("WorldPartitionEditor");
 				const float AutoCellLoadingMaxWorldSize = WorldPartitionEditorModule.GetAutoCellLoadingMaxWorldSize();
-				FVector WorldSize = GetWorldBounds().GetSize();
+				FVector WorldSize = GetEditorWorldBounds().GetSize();
 				const bool bItsASmallWorld = WorldSize.X <= AutoCellLoadingMaxWorldSize && WorldSize.Y <= AutoCellLoadingMaxWorldSize && WorldSize.Z <= AutoCellLoadingMaxWorldSize;
 			
 				// When loading a subworld, load all actors
