Fixed WorldPartition Editor Grid 2D now showing overlay box in PIE (only in SIE).

diff --git a/Engine/Source/Editor/WorldPartitionEditor/Private/WorldPartition/SWorldPartitionEditorGrid2D.cpp b/Engine/Source/Editor/WorldPartitionEditor/Private/WorldPartition/SWorldPartitionEditorGrid2D.cpp
--- a/Engine/Source/Editor/WorldPartitionEditor/Private/WorldPartition/SWorldPartitionEditorGrid2D.cpp
+++ b/Engine/Source/Editor/WorldPartitionEditor/Private/WorldPartition/SWorldPartitionEditorGrid2D.cpp
@@ -623,7 +623,7 @@ int32 SWorldPartitionEditorGrid2D::OnPaint(const FPaintArgs& Args, const FGeomet
 		LayerId = PaintSoftwareCursor(AllottedGeometry, MyCullingRect, OutDrawElements, ++LayerId);
 		
 		// Draw a surrounding indicator when PIE is active
-		if (UWorldPartition::IsSimulating())
+		if (UWorldPartition::IsSimulating() || !!GEditor->PlayWorld)
 		{
 			FSlateDrawElement::MakeBox(
 				OutDrawElements,
