World Partition editor now has realtime actor descriptors display through a custom actor desc view proxy.

diff --git a/Engine/Source/Editor/WorldPartitionEditor/Private/WorldPartition/SWorldPartitionEditorGrid2D.cpp b/Engine/Source/Editor/WorldPartitionEditor/Private/WorldPartition/SWorldPartitionEditorGrid2D.cpp
--- a/Engine/Source/Editor/WorldPartitionEditor/Private/WorldPartition/SWorldPartitionEditorGrid2D.cpp
+++ b/Engine/Source/Editor/WorldPartitionEditor/Private/WorldPartition/SWorldPartitionEditorGrid2D.cpp
@@ -9,6 +9,7 @@
 #include "WorldPartition/WorldPartitionEditorHash.h"
 #include "WorldPartition/WorldPartitionEditorCell.h"
 #include "WorldPartition/WorldPartitionActorDesc.h"
+#include "WorldPartition/WorldPartitionActorDescView.h"
 #include "Editor/GroupActor.h"
 #include "Editor/EditorEngine.h"
 #include "Widgets/Input/SCheckBox.h"
@@ -21,6 +22,36 @@
 
 #define LOCTEXT_NAMESPACE "WorldPartitionEditor"
 
+class FWorldPartitionActorDescViewBoundsProxy : public FWorldPartitionActorDescView
+{
+public:
+	FWorldPartitionActorDescViewBoundsProxy(const FWorldPartitionActorDesc* InActorDesc)
+	: FWorldPartitionActorDescView(InActorDesc)
+	{}
+
+	FVector GetOrigin() const
+	{
+		if (AActor* Actor = ActorDesc->GetActor())
+		{
+			return Actor->GetActorLocation();
+		}
+		return ActorDesc->GetOrigin();
+	}
+
+	FBox GetBounds() const
+	{
+		if (AActor* Actor = ActorDesc->GetActor())
+		{
+			FVector BoundsLocation;
+			FVector BoundsExtent;
+			Actor->GetActorLocationBounds(/*bOnlyCollidingComponents*/false, BoundsLocation, BoundsExtent, /*bIncludeFromChildActors*/true);
+			return FBox(BoundsLocation - BoundsExtent, BoundsLocation + BoundsExtent);
+		}
+
+		return ActorDesc->GetBounds();
+	}
+};
+
 static TAutoConsoleVariable<int32> CVarShowReloadMiniMapButton(TEXT("wp.MiniMap.ShowReloadButton"), 0, TEXT("Show reload MiniMap button."));
 
 SWorldPartitionEditorGrid2D::FEditorCommands::FEditorCommands()
@@ -339,14 +370,14 @@ uint32 SWorldPartitionEditorGrid2D::PaintActors(const FGeometry& AllottedGeometr
 	const FBox2D WorldViewRect(ScreenToWorld.TransformPoint(ViewRect.Min), ScreenToWorld.TransformPoint(ViewRect.Max));
 	const FBox ViewRectWorld(FVector(WorldViewRect.Min.X, WorldViewRect.Min.Y, -WORLD_MAX), FVector(WorldViewRect.Max.X, WorldViewRect.Max.Y, WORLD_MAX));
 
-	TSet<FWorldPartitionActorDesc*> ActorDescList;
+	TSet<FWorldPartitionActorDescViewBoundsProxy> ActorDescList;
 
 	// Include all actors if requested
 	if (bShowActors)
 	{
 		WorldPartition->EditorHash->ForEachIntersectingActor(ViewRectWorld, [&](FWorldPartitionActorDesc* ActorDesc)
 		{
-			ActorDescList.Add(ActorDesc);
+			ActorDescList.Emplace(ActorDesc);
 		});
 	}
 
@@ -358,7 +389,7 @@ uint32 SWorldPartitionEditorGrid2D::PaintActors(const FGeometry& AllottedGeometr
 		{
 			if (FWorldPartitionActorDesc* ActorDesc = WorldPartition->GetActorDesc(Actor->GetActorGuid()))
 			{
-				ActorDescList.Add(ActorDesc);
+				ActorDescList.Emplace(ActorDesc);
 				ActorDesc->Tag = FWorldPartitionActorDesc::GlobalTag;
 			}
 		}
@@ -373,9 +404,9 @@ uint32 SWorldPartitionEditorGrid2D::PaintActors(const FGeometry& AllottedGeometr
 		TArray<FVector2D> LinePoints;
 		LinePoints.SetNum(5);
 
-		for (FWorldPartitionActorDesc* ActorDesc: ActorDescList)
+		for (const FWorldPartitionActorDescViewBoundsProxy& ActorDescView: ActorDescList)
 		{
-			const FBox ActorBounds = ActorDesc->GetBounds();
+			const FBox ActorBounds = ActorDescView.GetBounds();
 			FVector Origin, Extent;
 			ActorBounds.GetCenterAndExtents(Origin, Extent);
 
@@ -397,17 +428,17 @@ uint32 SWorldPartitionEditorGrid2D::PaintActors(const FGeometry& AllottedGeometr
 			{
 				FPaintGeometry ActorGeometry = AllottedGeometry.ToPaintGeometry(TopLeft, BottomRight - TopLeft);
 				float ActorColorGradient = FMath::Min((ActorViewBox.GetArea() - MinimumAreaCull) / AreaFadeDistance, 1.0f);
-				float ActorBrightness = (ActorDesc->GetGridPlacement() == EActorGridPlacement::AlwaysLoaded) ? 0.3f : 1.0f;
+				float ActorBrightness = (ActorDescView.GetGridPlacement() == EActorGridPlacement::AlwaysLoaded) ? 0.3f : 1.0f;
 				FLinearColor ActorColor(ActorBrightness, ActorBrightness, ActorBrightness, ActorColorGradient);
 
-				UClass* ActorClass = ActorDesc->GetActorClass();
+				UClass* ActorClass = ActorDescView.GetActorClass();
 
 				const float SquaredDistanceToPoint = ActorViewBox.ComputeSquaredDistanceToPoint(MouseCursorPos);
-				if ((ActorDesc->Tag == FWorldPartitionActorDesc::GlobalTag) || (SquaredDistanceToPoint > 0.0f && SquaredDistanceToPoint <= 2.0f))
+				if ((ActorDescView.GetTag() == FWorldPartitionActorDesc::GlobalTag) || (SquaredDistanceToPoint > 0.0f && SquaredDistanceToPoint <= 2.0f))
 				{
 					ActorColor = FLinearColor::Yellow;
 
-					FName ActorLabel = ActorDesc->GetActorLabel();
+					FName ActorLabel = ActorDescView.GetActorLabel();
 					if (!ActorLabel.IsNone())
 					{
 						FSlateDrawElement::MakeText(
@@ -421,7 +452,7 @@ uint32 SWorldPartitionEditorGrid2D::PaintActors(const FGeometry& AllottedGeometr
 						);
 					}
 				}
-				else if ((SelectBox.GetVolume() > 0) && SelectBox.Intersect(ActorDesc->GetBounds()))
+				else if ((SelectBox.GetVolume() > 0) && SelectBox.Intersect(ActorDescView.GetBounds()))
 				{
 					ActorColor = FLinearColor::White;
 				}
