Moved EnableLoadingOfLastLoadedCells project setting to a per-user setting.

diff --git a/Engine/Source/Editor/WorldPartitionEditor/Private/WorldPartitionEditorModule.cpp b/Engine/Source/Editor/WorldPartitionEditor/Private/WorldPartitionEditorModule.cpp
--- a/Engine/Source/Editor/WorldPartitionEditor/Private/WorldPartitionEditorModule.cpp
+++ b/Engine/Source/Editor/WorldPartitionEditor/Private/WorldPartitionEditorModule.cpp
@@ -148,11 +148,6 @@ void FWorldPartitionEditorModule::SetConversionPromptEnabled(bool bEnabled)
 	GetMutableDefault<UWorldPartitionEditorSettings>()->bEnableConversionPrompt = bEnabled;
 }
 
-bool FWorldPartitionEditorModule::GetEnableLoadingOfLastLoadedCells() const
-{
-	return GetDefault<UWorldPartitionEditorSettings>()->bEnableLoadingOfLastLoadedCells;
-}
-	
 float FWorldPartitionEditorModule::GetAutoCellLoadingMaxWorldSize() const
 {
 	return GetDefault<UWorldPartitionEditorSettings>()->AutoCellLoadingMaxWorldSize;
diff --git a/Engine/Source/Editor/WorldPartitionEditor/Public/WorldPartition/WorldPartitionEditorSettings.h b/Engine/Source/Editor/WorldPartitionEditor/Public/WorldPartition/WorldPartitionEditorSettings.h
--- a/Engine/Source/Editor/WorldPartitionEditor/Public/WorldPartition/WorldPartitionEditorSettings.h
+++ b/Engine/Source/Editor/WorldPartitionEditor/Public/WorldPartition/WorldPartitionEditorSettings.h
@@ -25,9 +25,6 @@ public:
 	UPROPERTY(config, EditAnywhere, Category = MapConversion, meta = (ToolTip = "This option when enabled will show a conversion prompt when opening non World Partition maps"))
 	bool bEnableConversionPrompt = false;
 
-	UPROPERTY(config, EditAnywhere, Category = WorldPartition, meta = (ToolTip = "This option is used to enable the loading of the last loaded cells in the editor"))
-	bool bEnableLoadingOfLastLoadedCells = true;
-
 	UPROPERTY(config, EditAnywhere, Category = MapConversion, meta = (ToolTip = "Commandlet class to use for World Parition conversion"))
 	TSubclassOf<UWorldPartitionConvertCommandlet> CommandletClass;
 
diff --git a/Engine/Source/Editor/WorldPartitionEditor/Public/WorldPartitionEditorModule.h b/Engine/Source/Editor/WorldPartitionEditor/Public/WorldPartitionEditorModule.h
--- a/Engine/Source/Editor/WorldPartitionEditor/Public/WorldPartitionEditorModule.h
+++ b/Engine/Source/Editor/WorldPartitionEditor/Public/WorldPartitionEditorModule.h
@@ -40,11 +40,6 @@ public:
 	 */
 	virtual void SetConversionPromptEnabled(bool bEnabled) override;
 
-	/**
-	 *
-	 */
-	virtual bool GetEnableLoadingOfLastLoadedCells() const override;
-
 	/**
 	 *
 	 */
diff --git a/Engine/Source/Runtime/Engine/Private/WorldPartition/WorldPartition.cpp b/Engine/Source/Runtime/Engine/Private/WorldPartition/WorldPartition.cpp
--- a/Engine/Source/Runtime/Engine/Private/WorldPartition/WorldPartition.cpp
+++ b/Engine/Source/Runtime/Engine/Private/WorldPartition/WorldPartition.cpp
@@ -311,7 +311,7 @@ void UWorldPartition::Initialize(UWorld* InWorld, const FTransform& InTransform)
 			}
 
 			// Load last loaded cells
-			if (WorldPartitionEditorModule.GetEnableLoadingOfLastLoadedCells())
+			if (GetMutableDefault<UWorldPartitionEditorPerProjectUserSettings>()->GetEnableLoadingOfLastLoadedCells())
 			{
 				const TArray<FName>& EditorGridLastLoadedCells = GetMutableDefault<UWorldPartitionEditorPerProjectUserSettings>()->GetEditorGridLoadedCells(InWorld);
 
diff --git a/Engine/Source/Runtime/Engine/Public/WorldPartition/IWorldPartitionEditorModule.h b/Engine/Source/Runtime/Engine/Public/WorldPartition/IWorldPartitionEditorModule.h
--- a/Engine/Source/Runtime/Engine/Public/WorldPartition/IWorldPartitionEditorModule.h
+++ b/Engine/Source/Runtime/Engine/Public/WorldPartition/IWorldPartitionEditorModule.h
@@ -15,7 +15,6 @@ public:
 	virtual bool IsWorldPartitionEnabled() const = 0;
 	virtual bool IsConversionPromptEnabled() const = 0;
 	virtual void SetConversionPromptEnabled(bool bEnabled) = 0;
-	virtual bool GetEnableLoadingOfLastLoadedCells() const = 0;
 	virtual float GetAutoCellLoadingMaxWorldSize() const = 0;
 	virtual bool ConvertMap(const FString& InLongPackageName) = 0;
 
diff --git a/Engine/Source/Runtime/Engine/Public/WorldPartition/WorldPartitionEditorPerProjectUserSettings.h b/Engine/Source/Runtime/Engine/Public/WorldPartition/WorldPartitionEditorPerProjectUserSettings.h
--- a/Engine/Source/Runtime/Engine/Public/WorldPartition/WorldPartitionEditorPerProjectUserSettings.h
+++ b/Engine/Source/Runtime/Engine/Public/WorldPartition/WorldPartitionEditorPerProjectUserSettings.h
@@ -84,6 +84,11 @@ public:
 		}
 	}
 
+	bool GetEnableLoadingOfLastLoadedCells() const
+	{
+		return !bDisableLoadingOfLastLoadedCells;
+	}
+
 	const TArray<FName>& GetWorldDataLayersNotLoadedInEditor(UWorld* InWorld)
 	{
 		return PerWorldEditorSettings.FindOrAdd(TSoftObjectPtr<UWorld>(InWorld)).NotLoadedDataLayers;
@@ -108,6 +113,9 @@ private:
 	UPROPERTY(config)
 	uint32 bShowDataLayerContent : 1;
 
+	UPROPERTY(config)
+	uint32 bDisableLoadingOfLastLoadedCells : 1;
+
 	UPROPERTY(config)
 	TMap<TSoftObjectPtr<UWorld>, FWorldPartitionPerWorldSettings> PerWorldEditorSettings;
 #endif
