Boosted the minimum mouse cursor delta to determine a drag movement in World Partition Editor view.

diff --git a/Engine/Source/Editor/WorldPartitionEditor/Private/WorldPartition/SWorldPartitionEditorGrid2D.cpp b/Engine/Source/Editor/WorldPartitionEditor/Private/WorldPartition/SWorldPartitionEditorGrid2D.cpp
--- a/Engine/Source/Editor/WorldPartitionEditor/Private/WorldPartition/SWorldPartitionEditorGrid2D.cpp
+++ b/Engine/Source/Editor/WorldPartitionEditor/Private/WorldPartition/SWorldPartitionEditorGrid2D.cpp
@@ -300,7 +300,7 @@ FReply SWorldPartitionEditorGrid2D::OnMouseMove(const FGeometry& MyGeometry, con
 			}
 		}
 
-		if (bIsDragging || (bIsRightMouseButtonDown && (CursorDelta.Size() > 0.0f)))
+		if (bIsDragging || (bIsRightMouseButtonDown && (CursorDelta.Size() > 1.0f)))
 		{
 			bIsDragging = true;
 			Trans += ScreenToWorld.TransformVector(CursorDelta);
