World Partition in-editor convert
- By default, don't skip GUID validation.
- Made the commandlet class and the skip GUID validation advanced settings for the in-editor World Partition convert UI.
- If the bGenerateIni option is enabled, automatically open windows explorer in the folder containing the resulting ini file.

diff --git a/Engine/Source/Editor/WorldPartitionEditor/Private/WorldPartitionEditorModule.cpp b/Engine/Source/Editor/WorldPartitionEditor/Private/WorldPartitionEditorModule.cpp
--- a/Engine/Source/Editor/WorldPartitionEditor/Private/WorldPartitionEditorModule.cpp
+++ b/Engine/Source/Editor/WorldPartitionEditor/Private/WorldPartitionEditorModule.cpp
@@ -174,7 +174,7 @@ bool FWorldPartitionEditorModule::ConvertMap(const FString& InLongPackageName)
 	UWorldPartitionConvertOptions* DefaultConvertOptions = GetMutableDefault<UWorldPartitionConvertOptions>();
 	DefaultConvertOptions->CommandletClass = GetDefault<UWorldPartitionEditorSettings>()->CommandletClass;
 	DefaultConvertOptions->bInPlace = false;
-	DefaultConvertOptions->bSkipStableGUIDValidation = true;
+	DefaultConvertOptions->bSkipStableGUIDValidation = false;
 	DefaultConvertOptions->LongPackageName = InLongPackageName;
 
 	TSharedPtr<SWindow> DlgWindow =
@@ -241,6 +241,15 @@ bool FWorldPartitionEditorModule::ConvertMap(const FString& InLongPackageName)
 			if (Result == 0)
 			{
 				FMessageDialog::Open(EAppMsgType::Ok, LOCTEXT("ConvertMapCompleted", "Conversion succeeded. Editor will now restart."));
+
+#if	PLATFORM_DESKTOP
+				if (DefaultConvertOptions->bGenerateIni)
+				{
+					const FString PackageFilename = FPackageName::LongPackageNameToFilename(DefaultConvertOptions->LongPackageName);
+					const FString PackageDirectory = FPaths::ConvertRelativePathToFull(FPaths::GetPath(PackageFilename));
+					FPlatformProcess::ExploreFolder(*PackageDirectory);
+				}
+#endif
 				FUnrealEdMisc::Get().RestartEditor(false);
 			}
 		}
diff --git a/Engine/Source/Editor/WorldPartitionEditor/Public/WorldPartition/WorldPartitionConvertOptions.h b/Engine/Source/Editor/WorldPartitionEditor/Public/WorldPartition/WorldPartitionConvertOptions.h
--- a/Engine/Source/Editor/WorldPartitionEditor/Public/WorldPartition/WorldPartitionConvertOptions.h
+++ b/Engine/Source/Editor/WorldPartitionEditor/Public/WorldPartition/WorldPartitionConvertOptions.h
@@ -12,31 +12,29 @@ class UWorldPartitionConvertOptions : public UObject
 	GENERATED_BODY()
 
 public:
-
-	UWorldPartitionConvertOptions()
-	{
-	}
-
 	FString ToCommandletArgs() const;
 
-	/** Properties which are supposed to be baked out for the material(s) */
-	UPROPERTY(EditAnywhere, Category=Convert)
+	UPROPERTY(EditAnywhere, AdvancedDisplay, Category=Convert)
 	TSubclassOf<UWorldPartitionConvertCommandlet> CommandletClass;
 		
-	/** Determines whether to not allow usage of the source mesh data while baking out material properties */
 	UPROPERTY(EditAnywhere, Category=Convert, meta = (ToolTip = "Wether the conversion should create a new map with a _WP suffix or overwrite the source map"))
 	bool bInPlace;
 
-	UPROPERTY(EditAnywhere, AdvancedDisplay, Category=Convert)
+	UPROPERTY(EditAnywhere, AdvancedDisplay, Category = Convert)
 	bool bDeleteSourceLevels;
+
 	UPROPERTY(EditAnywhere, AdvancedDisplay, Category = Convert)
 	bool bGenerateIni;
+
 	UPROPERTY(EditAnywhere, AdvancedDisplay, Category = Convert)
 	bool bReportOnly;
+
 	UPROPERTY(EditAnywhere, AdvancedDisplay, Category = Convert)
 	bool bVerbose;
+
 	UPROPERTY(EditAnywhere, AdvancedDisplay, Category = Convert)
 	bool bSkipStableGUIDValidation;
+
 	UPROPERTY(EditAnywhere, AdvancedDisplay, Category = Convert)
 	bool bSkipMiniMapGeneration;
 
