Implemented AActor::GetStreamingBounds that returns the actor bounds relevant to streaming.
- For editor-only actors this returns all the components bounds, otherwise it only includes non-editor components.
- Always skip transient components in the streaming bounds calculation.
- This fixes lights having a null bounds for streaming.

diff --git a/Engine/Source/Editor/UnrealEd/Classes/Editor/GroupActor.h b/Engine/Source/Editor/UnrealEd/Classes/Editor/GroupActor.h
--- a/Engine/Source/Editor/UnrealEd/Classes/Editor/GroupActor.h
+++ b/Engine/Source/Editor/UnrealEd/Classes/Editor/GroupActor.h
@@ -44,8 +44,10 @@ class AGroupActor : public AActor
 	virtual void PostActorCreated() override;
 	virtual void SetIsTemporarilyHiddenInEditor( bool bIsHidden ) override;
 	virtual void GetActorBounds(bool bOnlyCollidingComponents, FVector& Origin, FVector& BoxExtent, bool bIncludeFromChildActors) const override;
-	virtual void GetActorLocationBounds(bool bOnlyCollidingComponents, FVector& Origin, FVector& BoxExtent, bool bIncludeFromChildActors) const override;
 	virtual EActorGridPlacement GetDefaultGridPlacement() const override { return EActorGridPlacement::None; }
+#if WITH_EDITOR
+	virtual FBox GetStreamingBounds() const override;
+#endif
 	//~ End AActor Interface
 
 	// @todo document
diff --git a/Engine/Source/Editor/UnrealEd/Private/Commandlets/WorldPartitionConvertCommandlet.cpp b/Engine/Source/Editor/UnrealEd/Private/Commandlets/WorldPartitionConvertCommandlet.cpp
--- a/Engine/Source/Editor/UnrealEd/Private/Commandlets/WorldPartitionConvertCommandlet.cpp
+++ b/Engine/Source/Editor/UnrealEd/Private/Commandlets/WorldPartitionConvertCommandlet.cpp
@@ -989,11 +989,7 @@ int32 UWorldPartitionConvertCommandlet::Main(const FString& Params)
 					// Only override default grid placement on actors that are not marked as always loaded
 					else if (Actor->GetGridPlacement() != EActorGridPlacement::AlwaysLoaded)
 					{
-						FVector BoundsLocation;
-						FVector BoundsExtent;
-						Actor->GetActorLocationBounds(/*bOnlyCollidingComponents*/false, BoundsLocation, BoundsExtent, /*bIncludeFromChildActors*/true);
-
-						const FBox ActorBounds(BoundsLocation - BoundsExtent, BoundsLocation + BoundsExtent);
+						const FBox ActorBounds = Actor->GetStreamingBounds();
 
 						if (!WorldBounds.IsInside(ActorBounds))
 						{
diff --git a/Engine/Source/Editor/UnrealEd/Private/GroupActor.cpp b/Engine/Source/Editor/UnrealEd/Private/GroupActor.cpp
--- a/Engine/Source/Editor/UnrealEd/Private/GroupActor.cpp
+++ b/Engine/Source/Editor/UnrealEd/Private/GroupActor.cpp
@@ -252,10 +252,30 @@ void AGroupActor::GetActorBounds(bool bOnlyCollidingComponents, FVector& Origin,
 	Bounds.GetCenterAndExtents(Origin, BoxExtent);
 }
 
-void AGroupActor::GetActorLocationBounds(bool bOnlyCollidingComponents, FVector& Origin, FVector& BoxExtent, bool bIncludeFromChildActors) const
+#if WITH_EDITOR
+FBox AGroupActor::GetStreamingBounds() const
 {
-	GetActorBounds(bOnlyCollidingComponents, Origin, BoxExtent, bIncludeFromChildActors);
+	FBox StreamingBounds = Super::GetStreamingBounds();
+
+	for (AActor* Actor : GroupActors)
+	{
+		if (Actor)
+		{
+			StreamingBounds += Actor->GetStreamingBounds();
+		}
+	}
+
+	for (AGroupActor* SubGroupActor : SubGroups)
+	{
+		if (SubGroupActor)
+		{
+			StreamingBounds += SubGroupActor->GetStreamingBounds();
+		}
+	}
+
+	return StreamingBounds;
 }
+#endif
 
 void GetBoundingVectorsForGroup(AGroupActor* GroupActor, FViewport* Viewport, FVector& OutVectorMin, FVector& OutVectorMax)
 {
diff --git a/Engine/Source/Editor/WorldPartitionEditor/Private/WorldPartition/SWorldPartitionEditorGrid2D.cpp b/Engine/Source/Editor/WorldPartitionEditor/Private/WorldPartition/SWorldPartitionEditorGrid2D.cpp
--- a/Engine/Source/Editor/WorldPartitionEditor/Private/WorldPartition/SWorldPartitionEditorGrid2D.cpp
+++ b/Engine/Source/Editor/WorldPartitionEditor/Private/WorldPartition/SWorldPartitionEditorGrid2D.cpp
@@ -29,23 +29,11 @@ public:
 	: FWorldPartitionActorDescView(InActorDesc)
 	{}
 
-	FVector GetOrigin() const
-	{
-		if (AActor* Actor = ActorDesc->GetActor())
-		{
-			return Actor->GetActorLocation();
-		}
-		return ActorDesc->GetOrigin();
-	}
-
 	FBox GetBounds() const
 	{
 		if (AActor* Actor = ActorDesc->GetActor())
 		{
-			FVector BoundsLocation;
-			FVector BoundsExtent;
-			Actor->GetActorLocationBounds(/*bOnlyCollidingComponents*/false, BoundsLocation, BoundsExtent, /*bIncludeFromChildActors*/true);
-			return FBox(BoundsLocation - BoundsExtent, BoundsLocation + BoundsExtent);
+			return Actor->GetStreamingBounds();
 		}
 
 		return ActorDesc->GetBounds();
@@ -424,7 +412,7 @@ uint32 SWorldPartitionEditorGrid2D::PaintActors(const FGeometry& AllottedGeometr
 
 			const float MinimumAreaCull = 32.0f;
 			const float AreaFadeDistance = 128.0f;
-			if (ActorViewBox.Intersect(ViewRect) && (ActorViewBox.GetArea() > MinimumAreaCull))
+			if (ActorViewBox.Intersect(ViewRect) && ((Extent.Size2D() < KINDA_SMALL_NUMBER) || ActorViewBox.GetArea() > MinimumAreaCull))
 			{
 				FPaintGeometry ActorGeometry = AllottedGeometry.ToPaintGeometry(TopLeft, BottomRight - TopLeft);
 				float ActorColorGradient = FMath::Min((ActorViewBox.GetArea() - MinimumAreaCull) / AreaFadeDistance, 1.0f);
@@ -444,7 +432,7 @@ uint32 SWorldPartitionEditorGrid2D::PaintActors(const FGeometry& AllottedGeometr
 						FSlateDrawElement::MakeText(
 							OutDrawElements,
 							++LayerId,
-							ActorGeometry,
+							AllottedGeometry.ToPaintGeometry(TopLeft, FVector2D(1,1)),
 							ActorLabel.ToString(),
 							SmallLayoutFont,
 							ESlateDrawEffect::None,
@@ -697,11 +685,7 @@ FReply SWorldPartitionEditorGrid2D::FocusSelection()
 		{
 			if (AActor* Actor = Cast<AActor>(*It))
 			{
-				FVector Origin;
-				FVector Extent;
-				Actor->GetActorBounds(false, Origin, Extent);
-
-				SelectionBox += FBox(Origin - Extent, Origin + Extent);
+				SelectionBox += Actor->GetStreamingBounds();
 			}
 		}
 	}
@@ -718,16 +702,17 @@ void SWorldPartitionEditorGrid2D::FocusBox(const FBox& Box) const
 {
 	check(ScreenRect.bIsValid);
 
-	if (Box.GetVolume() > 0.0f)
+	const FBox2D Box2D(FVector2D(Box.Min), FVector2D(Box.Max));
+	Trans = -FVector2D(Box2D.GetCenter());
+
+	if (Box2D.GetArea() > 0.0f)
 	{
 		const FVector2D ScreenExtent = ScreenRect.GetExtent();
-		const FVector2D SelectExtent = FVector2D(Box.GetExtent());
-
-		Trans = -FVector2D(Box.GetCenter());
+		const FVector2D SelectExtent = FVector2D(Box2D.GetExtent());
 		Scale = (ScreenExtent / SelectExtent).GetMin() * 0.75f;
+	}
 
-		UpdateTransform();
-	}	
+	UpdateTransform();
 }
 
 void SWorldPartitionEditorGrid2D::UpdateTransform() const
@@ -754,4 +739,4 @@ void SWorldPartitionEditorGrid2D::UpdateSelection()
 	}
 }
 
-#undef LOCTEXT_NAMESPACE
\ No newline at end of file
+#undef LOCTEXT_NAMESPACE
diff --git a/Engine/Source/Runtime/Core/Public/UObject/UE5MainStreamObjectVersion.h b/Engine/Source/Runtime/Core/Public/UObject/UE5MainStreamObjectVersion.h
--- a/Engine/Source/Runtime/Core/Public/UObject/UE5MainStreamObjectVersion.h
+++ b/Engine/Source/Runtime/Core/Public/UObject/UE5MainStreamObjectVersion.h
@@ -130,6 +130,9 @@ struct CORE_API FUE5MainStreamObjectVersion
 		// Changed how world partition streaming cells are named
 		WorldPartitionStreamingCellsNamingShortened,
 
+		// Changed how actor descriptors compute their bounds
+		WorldPartitionActorDescGetStreamingBounds,
+
         // -----<new versions can be added above this line>-------------------------------------------------
 		VersionPlusOne,
 		LatestVersion = VersionPlusOne - 1
diff --git a/Engine/Source/Runtime/Engine/Classes/Components/ActorComponent.h b/Engine/Source/Runtime/Engine/Classes/Components/ActorComponent.h
--- a/Engine/Source/Runtime/Engine/Classes/Components/ActorComponent.h
+++ b/Engine/Source/Runtime/Engine/Classes/Components/ActorComponent.h
@@ -884,6 +884,7 @@ public:
 	virtual void PostEditUndo() override;
 	virtual bool IsSelectedInEditor() const override;
 	virtual void SetPackageExternal(bool bExternal, bool bShouldDirty) {}
+	virtual FBox GetStreamingBounds() const { return FBox(ForceInit); }
 #endif // WITH_EDITOR
 	//~ End UObject Interface.
 
diff --git a/Engine/Source/Runtime/Engine/Classes/Components/LightComponent.h b/Engine/Source/Runtime/Engine/Classes/Components/LightComponent.h
--- a/Engine/Source/Runtime/Engine/Classes/Components/LightComponent.h
+++ b/Engine/Source/Runtime/Engine/Classes/Components/LightComponent.h
@@ -347,6 +347,13 @@ public:
 		return FSphere(FVector::ZeroVector, WORLD_MAX);
 	}
 
+#if WITH_EDITOR
+	virtual FBox GetStreamingBounds() const override
+	{
+		return GetBoundingBox();
+	}
+#endif // WITH_EDITOR
+
 	/**
 	 * Return the homogenous position of the light.
 	 */
diff --git a/Engine/Source/Runtime/Engine/Classes/Components/SceneComponent.h b/Engine/Source/Runtime/Engine/Classes/Components/SceneComponent.h
--- a/Engine/Source/Runtime/Engine/Classes/Components/SceneComponent.h
+++ b/Engine/Source/Runtime/Engine/Classes/Components/SceneComponent.h
@@ -868,6 +868,10 @@ public:
 	virtual void OnComponentDestroyed(bool bDestroyingHierarchy) override;
 	virtual void ApplyWorldOffset(const FVector& InOffset, bool bWorldShift) override;
 	virtual TStructOnScope<FActorComponentInstanceData> GetComponentInstanceData() const override;
+#if WITH_EDITOR
+	virtual FBox GetStreamingBounds() const override;
+#endif // WITH_EDITOR
+
 	//~ End ActorComponent Interface
 
 	//~ Begin UObject Interface
diff --git a/Engine/Source/Runtime/Engine/Classes/Components/SpotLightComponent.h b/Engine/Source/Runtime/Engine/Classes/Components/SpotLightComponent.h
--- a/Engine/Source/Runtime/Engine/Classes/Components/SpotLightComponent.h
+++ b/Engine/Source/Runtime/Engine/Classes/Components/SpotLightComponent.h
@@ -40,6 +40,7 @@ class ENGINE_API USpotLightComponent : public UPointLightComponent
 	virtual float ComputeLightBrightness() const override;
 #if WITH_EDITOR
 	virtual void SetLightBrightness(float InBrightness) override;
+	virtual FBox GetStreamingBounds() const override;
 #endif
 
 	float GetHalfConeAngle() const;
diff --git a/Engine/Source/Runtime/Engine/Classes/GameFramework/Actor.h b/Engine/Source/Runtime/Engine/Classes/GameFramework/Actor.h
--- a/Engine/Source/Runtime/Engine/Classes/GameFramework/Actor.h
+++ b/Engine/Source/Runtime/Engine/Classes/GameFramework/Actor.h
@@ -2052,7 +2052,7 @@ public:
 	 *
 	 * @see GetActorBounds()
 	 */
-	virtual void GetActorLocationBounds(bool bOnlyCollidingComponents, FVector& Origin, FVector& BoxExtent, bool bIncludeFromChildActors = false) const;
+	virtual FBox GetStreamingBounds() const;
 #endif
 
 
diff --git a/Engine/Source/Runtime/Engine/Private/Actor.cpp b/Engine/Source/Runtime/Engine/Private/Actor.cpp
--- a/Engine/Source/Runtime/Engine/Private/Actor.cpp
+++ b/Engine/Source/Runtime/Engine/Private/Actor.cpp
@@ -983,18 +983,50 @@ void AActor::ProcessEvent(UFunction* Function, void* Parameters)
 }
 
 #if WITH_EDITOR
-void AActor::GetActorLocationBounds(bool bOnlyCollidingComponents, FVector& Origin, FVector& BoxExtent, bool bIncludeFromChildActors) const
+FBox AActor::GetStreamingBounds() const
 {
-	// Only PrimitiveComponents provides bounds contribution
-	if (FindComponentByClass<UPrimitiveComponent>())
+	FBox StreamingBounds(ForceInit);
+
+	auto IsComponentRelevant = [this](const UActorComponent* Component)
 	{
-		GetActorBounds(bOnlyCollidingComponents, Origin, BoxExtent, bIncludeFromChildActors);
-	}
-	else
+		if (!Component->IsRegistered())
+		{
+			return false;
+		}
+
+		if (Component->HasAnyFlags(RF_Transient))
+		{
+			return false;
+		}
+
+		if (!IsEditorOnly() && Component->IsEditorOnly())
+		{
+			return false;
+		}
+
+		return true;
+	};
+
+	ForEachComponent<UPrimitiveComponent>(true, [this, &StreamingBounds, &IsComponentRelevant](UActorComponent* Component)
+	{
+		if (IsComponentRelevant(Component))
+		{
+			StreamingBounds += Component->GetStreamingBounds();
+		}
+	});
+
+	if (!StreamingBounds.IsValid)
 	{
-		Origin = GetActorLocation();
-		BoxExtent = FVector(ForceInitToZero);
+		for (UActorComponent* Component : OwnedComponents)
+		{
+			if (IsComponentRelevant(Component))
+			{
+				StreamingBounds += Component->GetStreamingBounds();
+			}
+		}
 	}
+
+	return StreamingBounds;
 }
 #endif
 
diff --git a/Engine/Source/Runtime/Engine/Private/Components/SceneComponent.cpp b/Engine/Source/Runtime/Engine/Private/Components/SceneComponent.cpp
--- a/Engine/Source/Runtime/Engine/Private/Components/SceneComponent.cpp
+++ b/Engine/Source/Runtime/Engine/Private/Components/SceneComponent.cpp
@@ -2343,6 +2343,13 @@ TStructOnScope<FActorComponentInstanceData> USceneComponent::GetComponentInstanc
 	return MakeStructOnScope<FActorComponentInstanceData, FSceneComponentInstanceData>(this);;
 }
 
+#if WITH_EDITOR
+FBox USceneComponent::GetStreamingBounds() const
+{
+	return Bounds.GetBox();
+}
+#endif
+
 void USceneComponent::UpdateChildTransforms(EUpdateTransformFlags UpdateTransformFlags, ETeleportType Teleport)
 {
 	SCOPE_CYCLE_COUNTER(STAT_UpdateChildTransforms);
diff --git a/Engine/Source/Runtime/Engine/Private/Components/SpotLightComponent.cpp b/Engine/Source/Runtime/Engine/Private/Components/SpotLightComponent.cpp
--- a/Engine/Source/Runtime/Engine/Private/Components/SpotLightComponent.cpp
+++ b/Engine/Source/Runtime/Engine/Private/Components/SpotLightComponent.cpp
@@ -242,6 +242,12 @@ void USpotLightComponent::SetLightBrightness(float InBrightness)
 		ULightComponent::SetLightBrightness(InBrightness);
 	}
 }
+
+FBox USpotLightComponent::GetStreamingBounds() const
+{
+	const FSphere BoundingSphere = GetBoundingSphere();
+	return FBox(BoundingSphere.Center - BoundingSphere.W, BoundingSphere.Center + BoundingSphere.W);
+}
 #endif // WITH_EDITOR
 
 static bool IsSpotLightSupported(const USpotLightComponent* InLight)
diff --git a/Engine/Source/Runtime/Engine/Private/LevelInstance/LevelInstanceActor.cpp b/Engine/Source/Runtime/Engine/Private/LevelInstance/LevelInstanceActor.cpp
--- a/Engine/Source/Runtime/Engine/Private/LevelInstance/LevelInstanceActor.cpp
+++ b/Engine/Source/Runtime/Engine/Private/LevelInstance/LevelInstanceActor.cpp
@@ -618,10 +618,8 @@ void ALevelInstance::OnLevelInstanceLoaded()
 	}
 }
 
-void ALevelInstance::GetActorLocationBounds(bool bOnlyCollidingComponents, FVector& Origin, FVector& BoxExtent, bool bIncludeFromChildActors) const
+FBox ALevelInstance::GetStreamingBounds() const
 {
-	Super::GetActorLocationBounds(bOnlyCollidingComponents, Origin, BoxExtent, bIncludeFromChildActors);
-
 	// Add Level Bounds
 	if (SupportsLoading())
 	{
@@ -630,10 +628,12 @@ void ALevelInstance::GetActorLocationBounds(bool bOnlyCollidingComponents, FVect
 			FBox LevelInstanceBounds;
 			if (LevelInstanceSubsystem->GetLevelInstanceBounds(this, LevelInstanceBounds))
 			{
-				LevelInstanceBounds.GetCenterAndExtents(Origin, BoxExtent);
+				return LevelInstanceBounds;
 			}
 		}
 	}
+
+	return Super::GetStreamingBounds();
 }
 
 bool ALevelInstance::IsLockLocation() const
diff --git a/Engine/Source/Runtime/Engine/Private/WorldPartition/HLOD/HLODActor.cpp b/Engine/Source/Runtime/Engine/Private/WorldPartition/HLOD/HLODActor.cpp
--- a/Engine/Source/Runtime/Engine/Private/WorldPartition/HLOD/HLODActor.cpp
+++ b/Engine/Source/Runtime/Engine/Private/WorldPartition/HLOD/HLODActor.cpp
@@ -156,9 +156,9 @@ void AWorldPartitionHLOD::GetActorBounds(bool bOnlyCollidingComponents, FVector&
 	HLODBounds.GetCenterAndExtents(Origin, BoxExtent);
 }
 
-void AWorldPartitionHLOD::GetActorLocationBounds(bool bOnlyCollidingComponents, FVector& Origin, FVector& BoxExtent, bool bIncludeFromChildActors) const
+FBox AWorldPartitionHLOD::GetStreamingBounds() const
 {
-	GetActorBounds(bOnlyCollidingComponents, Origin, BoxExtent, bIncludeFromChildActors);
+	return HLODBounds;
 }
 
 uint32 AWorldPartitionHLOD::GetHLODHash() const
diff --git a/Engine/Source/Runtime/Engine/Private/WorldPartition/NavigationData/NavigationDataChunkActor.cpp b/Engine/Source/Runtime/Engine/Private/WorldPartition/NavigationData/NavigationDataChunkActor.cpp
--- a/Engine/Source/Runtime/Engine/Private/WorldPartition/NavigationData/NavigationDataChunkActor.cpp
+++ b/Engine/Source/Runtime/Engine/Private/WorldPartition/NavigationData/NavigationDataChunkActor.cpp
@@ -129,8 +129,8 @@ EActorGridPlacement ANavigationDataChunkActor::GetDefaultGridPlacement() const
 	return EActorGridPlacement::Bounds;
 }
 
-void ANavigationDataChunkActor::GetActorLocationBounds(bool bOnlyCollidingComponents, FVector& OutOrigin, FVector& OutBoxExtent, bool bIncludeFromChildActors) const
+FBox ANavigationDataChunkActor::GetStreamingBounds() const
 {
-	GetActorBounds(bOnlyCollidingComponents, OutOrigin, OutBoxExtent, bIncludeFromChildActors);
+	return DataChunkActorBounds;
 }
 #endif // WITH_EDITOR
diff --git a/Engine/Source/Runtime/Engine/Private/WorldPartition/WorldPartitionActorDesc.cpp b/Engine/Source/Runtime/Engine/Private/WorldPartition/WorldPartitionActorDesc.cpp
--- a/Engine/Source/Runtime/Engine/Private/WorldPartition/WorldPartitionActorDesc.cpp
+++ b/Engine/Source/Runtime/Engine/Private/WorldPartition/WorldPartitionActorDesc.cpp
@@ -43,12 +43,8 @@ void FWorldPartitionActorDesc::Init(const AActor* InActor)
 
 	ActorPtr = (AActor*)InActor;
 
-	FVector NewBoundsLocation;
-	FVector NewBoundsExtent;
-	InActor->GetActorLocationBounds(/*bOnlyCollidingComponents*/false, NewBoundsLocation, NewBoundsExtent, /*bIncludeFromChildActors*/true);
-
-	BoundsLocation = NewBoundsLocation;
-	BoundsExtent = NewBoundsExtent;
+	const FBox StreamingBounds = InActor->GetStreamingBounds();
+	StreamingBounds.GetCenterAndExtents(BoundsLocation, BoundsExtent);
 
 	const EActorGridPlacement DefaultGridPlacement = InActor->GetDefaultGridPlacement();
 	if (DefaultGridPlacement != EActorGridPlacement::None)
diff --git a/Engine/Source/Runtime/Engine/Public/LevelInstance/LevelInstanceActor.h b/Engine/Source/Runtime/Engine/Public/LevelInstance/LevelInstanceActor.h
--- a/Engine/Source/Runtime/Engine/Public/LevelInstance/LevelInstanceActor.h
+++ b/Engine/Source/Runtime/Engine/Public/LevelInstance/LevelInstanceActor.h
@@ -72,7 +72,7 @@ public:
 	virtual void PushSelectionToProxies() override;
 	virtual void PushLevelInstanceEditingStateToProxies(bool bInEditingState) override;
 	virtual FBox GetComponentsBoundingBox(bool bNonColliding = false, bool bIncludeFromChildActors = false) const override;
-	virtual void GetActorLocationBounds(bool bOnlyCollidingComponents, FVector& Origin, FVector& BoxExtent, bool bIncludeFromChildActors = false) const override;
+	virtual FBox GetStreamingBounds() const override;
 	virtual bool IsLockLocation() const override;
 	virtual ELevelInstanceRuntimeBehavior GetDesiredRuntimeBehavior() const { return DesiredRuntimeBehavior; }
 	virtual ELevelInstanceRuntimeBehavior GetDefaultRuntimeBehavior() const { return ELevelInstanceRuntimeBehavior::Embedded; }
diff --git a/Engine/Source/Runtime/Engine/Public/WorldPartition/HLOD/HLODActor.h b/Engine/Source/Runtime/Engine/Public/WorldPartition/HLOD/HLODActor.h
--- a/Engine/Source/Runtime/Engine/Public/WorldPartition/HLOD/HLODActor.h
+++ b/Engine/Source/Runtime/Engine/Public/WorldPartition/HLOD/HLODActor.h
@@ -73,7 +73,7 @@ protected:
 	virtual TUniquePtr<class FWorldPartitionActorDesc> CreateClassActorDesc() const override;
 
 	virtual void GetActorBounds(bool bOnlyCollidingComponents, FVector& Origin, FVector& BoxExtent, bool bIncludeFromChildActors) const override;
-	virtual void GetActorLocationBounds(bool bOnlyCollidingComponents, FVector& Origin, FVector& BoxExtent, bool bIncludeFromChildActors) const override;
+	virtual FBox GetStreamingBounds() const override;
 #endif
 	//~ End AActor Interface.
 
diff --git a/Engine/Source/Runtime/Engine/Public/WorldPartition/NavigationData/NavigationDataChunkActor.h b/Engine/Source/Runtime/Engine/Public/WorldPartition/NavigationData/NavigationDataChunkActor.h
--- a/Engine/Source/Runtime/Engine/Public/WorldPartition/NavigationData/NavigationDataChunkActor.h
+++ b/Engine/Source/Runtime/Engine/Public/WorldPartition/NavigationData/NavigationDataChunkActor.h
@@ -31,7 +31,7 @@ public:
 
 	//~ Begin AActor Interface.
 	virtual EActorGridPlacement GetDefaultGridPlacement() const override;
-	virtual void GetActorLocationBounds(bool bOnlyCollidingComponents, FVector& OutOrigin, FVector& OutBoxExtent, bool bIncludeFromChildActors) const override;
+	virtual FBox GetStreamingBounds() const override;
 	//~ End AActor Interface.
 #endif // WITH_EDITOR
 	
