Allow World Partition cell selection / manipulation with single-click

diff --git a/Engine/Source/Editor/WorldPartitionEditor/Private/WorldPartition/SWorldPartitionEditorGrid2D.cpp b/Engine/Source/Editor/WorldPartitionEditor/Private/WorldPartition/SWorldPartitionEditorGrid2D.cpp
--- a/Engine/Source/Editor/WorldPartitionEditor/Private/WorldPartition/SWorldPartitionEditorGrid2D.cpp
+++ b/Engine/Source/Editor/WorldPartitionEditor/Private/WorldPartition/SWorldPartitionEditorGrid2D.cpp
@@ -151,7 +151,14 @@ void SWorldPartitionEditorGrid2D::Construct(const FArguments& InArgs)
 
 	auto CanLoadOrUnloadCells = [this]()
 	{
-		return GLevelEditorModeTools().IsDefaultModeActive() && SelectBox.GetVolume() > 0;
+		TArray<UWorldPartitionEditorCell*> Cells;
+		
+		if (WorldPartition == nullptr)
+		{
+			return false;
+		}
+
+		return GLevelEditorModeTools().IsDefaultModeActive() && (WorldPartition->EditorHash->GetIntersectingCells(SelectBox, Cells) > 0);
 	};
 
 	ActionList.MapAction(Commands.LoadSelectedCells, FExecuteAction::CreateSP(this, &SWorldPartitionEditorGrid2D::LoadSelectedCells), FCanExecuteAction::CreateLambda(CanLoadOrUnloadCells));
@@ -733,10 +740,8 @@ void SWorldPartitionEditorGrid2D::UpdateSelection()
 	);
 
 	const FBox2D SelectBox2D(FVector2D(SelectionBox.Min), FVector2D(SelectionBox.Max));
-	if(!FMath::IsNearlyZero(SelectBox2D.GetArea()))
-	{
-		SelectBox = SelectionBox;
-	}
+
+	SelectBox = SelectionBox;
 }
 
 #undef LOCTEXT_NAMESPACE
diff --git a/Engine/Source/Editor/WorldPartitionEditor/Private/WorldPartition/SWorldPartitionEditorGridSpatialHash.cpp b/Engine/Source/Editor/WorldPartitionEditor/Private/WorldPartition/SWorldPartitionEditorGridSpatialHash.cpp
--- a/Engine/Source/Editor/WorldPartitionEditor/Private/WorldPartition/SWorldPartitionEditorGridSpatialHash.cpp
+++ b/Engine/Source/Editor/WorldPartitionEditor/Private/WorldPartition/SWorldPartitionEditorGridSpatialHash.cpp
@@ -215,47 +215,45 @@ int32 SWorldPartitionEditorGridSpatialHash::PaintGrid(const FGeometry& AllottedG
 	}
 
 	// Draw selected cells
-	if (SelectBox.GetVolume() > 0)
-	{
-		const FBox VisibleSelectBox = SelectBox.Overlap(VisibleGridRectWorld);
+	TMap<uint32, FCellDesc2D> UniqueCells2D;
+
+	const FBox VisibleSelectBox = SelectBox.Overlap(VisibleGridRectWorld);
 
-		if (VisibleSelectBox.GetVolume() > 0)
+	if (VisibleSelectBox.IsValid)
+	{
+		WorldPartition->EditorHash->ForEachIntersectingCell(VisibleSelectBox, [&](UWorldPartitionEditorCell* Cell)
 		{
-			TMap<uint32, FCellDesc2D> UniqueCells2D;
-			WorldPartition->EditorHash->ForEachIntersectingCell(VisibleSelectBox, [&](UWorldPartitionEditorCell* Cell)
-			{
-				UWorldPartitionEditorSpatialHash::FCellCoord CellCoord = EditorSpatialHash->GetCellCoords(Cell->Bounds.GetCenter(), 0);
+			UWorldPartitionEditorSpatialHash::FCellCoord CellCoord = EditorSpatialHash->GetCellCoords(Cell->Bounds.GetCenter(), 0);
+
+			FHashBuilder HashBuilder;
+			HashBuilder << CellCoord.X << CellCoord.Y;
+			uint32 CellHash2D = HashBuilder.GetHash();
 
-				FHashBuilder HashBuilder;
-				HashBuilder << CellCoord.X << CellCoord.Y;
-				uint32 CellHash2D = HashBuilder.GetHash();
+			FCellDesc2D& CellDesc2D = UniqueCells2D.Add(CellHash2D);
 
-				FCellDesc2D& CellDesc2D = UniqueCells2D.Add(CellHash2D);
+			CellDesc2D.Bounds = FBox2D(FVector2D(Cell->Bounds.Min), FVector2D(Cell->Bounds.Max));
+		});
 
-				CellDesc2D.Bounds = FBox2D(FVector2D(Cell->Bounds.Min), FVector2D(Cell->Bounds.Max));
-			});
+		for (auto& UniqueCell : UniqueCells2D)
+		{
+			const FCellDesc2D& Cell = UniqueCell.Value;
 
-			for(auto& UniqueCell: UniqueCells2D)
-			{
-				const FCellDesc2D& Cell = UniqueCell.Value;
-
-				FPaintGeometry CellGeometry = AllottedGeometry.ToPaintGeometry(
-					WorldToScreen.TransformPoint(Cell.Bounds.Min),
-					WorldToScreen.TransformPoint(Cell.Bounds.Max) - WorldToScreen.TransformPoint(Cell.Bounds.Min)
-				);
-
-				FSlateColorBrush CellBrush(FLinearColor::White);
-				FLinearColor CellColor(1, 1, 1, 0.25f);
-
-				FSlateDrawElement::MakeBox(
-					OutDrawElements,
-					++LayerId,
-					CellGeometry,
-					&CellBrush,
-					ESlateDrawEffect::None,
-					CellColor
-				);
-			}
+			FPaintGeometry CellGeometry = AllottedGeometry.ToPaintGeometry(
+				WorldToScreen.TransformPoint(Cell.Bounds.Min),
+				WorldToScreen.TransformPoint(Cell.Bounds.Max) - WorldToScreen.TransformPoint(Cell.Bounds.Min)
+			);
+
+			FSlateColorBrush CellBrush(FLinearColor::White);
+			FLinearColor CellColor(1, 1, 1, 0.25f);
+
+			FSlateDrawElement::MakeBox(
+				OutDrawElements,
+				++LayerId,
+				CellGeometry,
+				&CellBrush,
+				ESlateDrawEffect::None,
+				CellColor
+			);
 		}
 	}
 
