Fix always having 1 cell selected at 0,0,0 in World Partition Editor.
Fix display Load/Unload when no cell was selected
Invalid SelectBox was tested against VisibleSelectBox

diff --git a/Engine/Source/Editor/WorldPartitionEditor/Private/WorldPartition/SWorldPartitionEditorGrid2D.cpp b/Engine/Source/Editor/WorldPartitionEditor/Private/WorldPartition/SWorldPartitionEditorGrid2D.cpp
--- a/Engine/Source/Editor/WorldPartitionEditor/Private/WorldPartition/SWorldPartitionEditorGrid2D.cpp
+++ b/Engine/Source/Editor/WorldPartitionEditor/Private/WorldPartition/SWorldPartitionEditorGrid2D.cpp
@@ -159,7 +159,7 @@ void SWorldPartitionEditorGrid2D::Construct(const FArguments& InArgs)
 			return false;
 		}
 
-		return GLevelEditorModeTools().IsDefaultModeActive() && (WorldPartition->EditorHash->GetIntersectingCells(SelectBox, Cells) > 0);
+		return GLevelEditorModeTools().IsDefaultModeActive() && SelectBox.IsValid && (WorldPartition->EditorHash->GetIntersectingCells(SelectBox, Cells) > 0);
 	};
 
 	ActionList.MapAction(Commands.LoadSelectedCells, FExecuteAction::CreateSP(this, &SWorldPartitionEditorGrid2D::LoadSelectedCells), FCanExecuteAction::CreateLambda(CanLoadOrUnloadCells));
@@ -740,8 +740,6 @@ void SWorldPartitionEditorGrid2D::UpdateSelection()
 		FVector(FMath::Max(SelectionStart.X, SelectionEnd.X), FMath::Max(SelectionStart.Y, SelectionEnd.Y), WORLD_MAX)
 	);
 
-	const FBox2D SelectBox2D(FVector2D(SelectionBox.Min), FVector2D(SelectionBox.Max));
-
 	SelectBox = SelectionBox;
 }
 
diff --git a/Engine/Source/Editor/WorldPartitionEditor/Private/WorldPartition/SWorldPartitionEditorGridSpatialHash.cpp b/Engine/Source/Editor/WorldPartitionEditor/Private/WorldPartition/SWorldPartitionEditorGridSpatialHash.cpp
--- a/Engine/Source/Editor/WorldPartitionEditor/Private/WorldPartition/SWorldPartitionEditorGridSpatialHash.cpp
+++ b/Engine/Source/Editor/WorldPartitionEditor/Private/WorldPartition/SWorldPartitionEditorGridSpatialHash.cpp
@@ -215,45 +215,48 @@ int32 SWorldPartitionEditorGridSpatialHash::PaintGrid(const FGeometry& AllottedG
 	}
 
 	// Draw selected cells
-	TMap<uint32, FCellDesc2D> UniqueCells2D;
-
-	const FBox VisibleSelectBox = SelectBox.Overlap(VisibleGridRectWorld);
-
-	if (VisibleSelectBox.IsValid)
+	if (SelectBox.IsValid)
 	{
-		WorldPartition->EditorHash->ForEachIntersectingCell(VisibleSelectBox, [&](UWorldPartitionEditorCell* Cell)
+		const FBox VisibleSelectBox = SelectBox.Overlap(VisibleGridRectWorld);
+
+		if (VisibleSelectBox.IsValid)
 		{
-			UWorldPartitionEditorSpatialHash::FCellCoord CellCoord = EditorSpatialHash->GetCellCoords(Cell->Bounds.GetCenter(), 0);
+			TMap<uint32, FCellDesc2D> UniqueCells2D;
 
-			FHashBuilder HashBuilder;
-			HashBuilder << CellCoord.X << CellCoord.Y;
-			uint32 CellHash2D = HashBuilder.GetHash();
+			WorldPartition->EditorHash->ForEachIntersectingCell(VisibleSelectBox, [&](UWorldPartitionEditorCell* Cell)
+			{
+				UWorldPartitionEditorSpatialHash::FCellCoord CellCoord = EditorSpatialHash->GetCellCoords(Cell->Bounds.GetCenter(), 0);
 
-			FCellDesc2D& CellDesc2D = UniqueCells2D.Add(CellHash2D);
+				FHashBuilder HashBuilder;
+				HashBuilder << CellCoord.X << CellCoord.Y;
+				uint32 CellHash2D = HashBuilder.GetHash();
 
-			CellDesc2D.Bounds = FBox2D(FVector2D(Cell->Bounds.Min), FVector2D(Cell->Bounds.Max));
-		});
+				FCellDesc2D& CellDesc2D = UniqueCells2D.Add(CellHash2D);
 
-		for (auto& UniqueCell : UniqueCells2D)
-		{
-			const FCellDesc2D& Cell = UniqueCell.Value;
+				CellDesc2D.Bounds = FBox2D(FVector2D(Cell->Bounds.Min), FVector2D(Cell->Bounds.Max));
+			});
 
-			FPaintGeometry CellGeometry = AllottedGeometry.ToPaintGeometry(
-				WorldToScreen.TransformPoint(Cell.Bounds.Min),
-				WorldToScreen.TransformPoint(Cell.Bounds.Max) - WorldToScreen.TransformPoint(Cell.Bounds.Min)
-			);
-
-			FSlateColorBrush CellBrush(FLinearColor::White);
-			FLinearColor CellColor(1, 1, 1, 0.25f);
-
-			FSlateDrawElement::MakeBox(
-				OutDrawElements,
-				++LayerId,
-				CellGeometry,
-				&CellBrush,
-				ESlateDrawEffect::None,
-				CellColor
-			);
+			for (auto& UniqueCell : UniqueCells2D)
+			{
+				const FCellDesc2D& Cell = UniqueCell.Value;
+
+				FPaintGeometry CellGeometry = AllottedGeometry.ToPaintGeometry(
+					WorldToScreen.TransformPoint(Cell.Bounds.Min),
+					WorldToScreen.TransformPoint(Cell.Bounds.Max) - WorldToScreen.TransformPoint(Cell.Bounds.Min)
+				);
+
+				FSlateColorBrush CellBrush(FLinearColor::White);
+				FLinearColor CellColor(1, 1, 1, 0.25f);
+
+				FSlateDrawElement::MakeBox(
+					OutDrawElements,
+					++LayerId,
+					CellGeometry,
+					&CellBrush,
+					ESlateDrawEffect::None,
+					CellColor
+				);
+			}
 		}
 	}
 
