Allow changing editor cells loaded state for non-default editor modes.

diff --git a/Engine/Source/Editor/WorldPartitionEditor/Private/WorldPartition/SWorldPartitionEditorGrid2D.cpp b/Engine/Source/Editor/WorldPartitionEditor/Private/WorldPartition/SWorldPartitionEditorGrid2D.cpp
--- a/Engine/Source/Editor/WorldPartitionEditor/Private/WorldPartition/SWorldPartitionEditorGrid2D.cpp
+++ b/Engine/Source/Editor/WorldPartitionEditor/Private/WorldPartition/SWorldPartitionEditorGrid2D.cpp
@@ -17,8 +17,6 @@
 #include "WorldBrowserModule.h"
 #include "LevelEditorViewport.h"
 #include "Algo/Transform.h"
-#include "Editor.h"
-#include "EditorModeManager.h"
 
 #define LOCTEXT_NAMESPACE "WorldPartitionEditor"
 
@@ -152,14 +150,7 @@ void SWorldPartitionEditorGrid2D::Construct(const FArguments& InArgs)
 
 	auto CanLoadOrUnloadCells = [this]()
 	{
-		TArray<UWorldPartitionEditorCell*> Cells;
-		
-		if (WorldPartition == nullptr)
-		{
-			return false;
-		}
-
-		return GLevelEditorModeTools().IsDefaultModeActive() && SelectBox.IsValid && (WorldPartition->EditorHash->GetIntersectingCells(SelectBox, Cells) > 0);
+		return WorldPartition && SelectBox.IsValid && (WorldPartition->EditorHash->ForEachIntersectingCell(SelectBox, [&](UWorldPartitionEditorCell*){}) > 0);
 	};
 
 	ActionList.MapAction(Commands.LoadSelectedCells, FExecuteAction::CreateSP(this, &SWorldPartitionEditorGrid2D::LoadSelectedCells), FCanExecuteAction::CreateLambda(CanLoadOrUnloadCells));
