World Partition Editor: fix drag delta to match real mouse cursor movement.

diff --git a/Engine/Source/Editor/WorldPartitionEditor/Private/WorldPartition/SWorldPartitionEditorGrid2D.cpp b/Engine/Source/Editor/WorldPartitionEditor/Private/WorldPartition/SWorldPartitionEditorGrid2D.cpp
--- a/Engine/Source/Editor/WorldPartitionEditor/Private/WorldPartition/SWorldPartitionEditorGrid2D.cpp
+++ b/Engine/Source/Editor/WorldPartitionEditor/Private/WorldPartition/SWorldPartitionEditorGrid2D.cpp
@@ -289,8 +289,14 @@ FReply SWorldPartitionEditorGrid2D::OnMouseMove(const FGeometry& MyGeometry, con
 
 		if (bIsDragging || (bIsRightMouseButtonDown && (CursorDelta.Size() > 1.0f)))
 		{
-			bIsDragging = true;
-			Trans += ScreenToWorld.TransformVector(CursorDelta);
+			if (!bIsDragging)
+			{
+				LastMouseCursorPosWorldDrag = MouseCursorPosWorld;
+				bIsDragging = true;
+			}
+
+			Trans += (MouseCursorPosWorld - LastMouseCursorPosWorldDrag);
+
 			UpdateTransform();
 			return FReply::Handled();
 		}
diff --git a/Engine/Source/Editor/WorldPartitionEditor/Public/WorldPartition/SWorldPartitionEditorGrid2D.h b/Engine/Source/Editor/WorldPartitionEditor/Public/WorldPartition/SWorldPartitionEditorGrid2D.h
--- a/Engine/Source/Editor/WorldPartitionEditor/Public/WorldPartition/SWorldPartitionEditorGrid2D.h
+++ b/Engine/Source/Editor/WorldPartitionEditor/Public/WorldPartition/SWorldPartitionEditorGrid2D.h
@@ -89,6 +89,7 @@ protected:
 	bool bShowGroupActors;
 	FVector2D MouseCursorPos;
 	FVector2D MouseCursorPosWorld;
+	FVector2D LastMouseCursorPosWorldDrag;
 	FVector2D SelectionStart;
 	FVector2D SelectionEnd;
 	FBox SelectBox;
