World Partition: Replace conversion prompt with Tools > [World Partition] Convert Map...

diff --git a/Engine/Source/Editor/GameProjectGeneration/Private/GameProjectUtils.cpp b/Engine/Source/Editor/GameProjectGeneration/Private/GameProjectUtils.cpp
--- a/Engine/Source/Editor/GameProjectGeneration/Private/GameProjectUtils.cpp
+++ b/Engine/Source/Editor/GameProjectGeneration/Private/GameProjectUtils.cpp
@@ -271,12 +271,6 @@ namespace
 	{
 		if (InProjectInfo.bIsBlankTemplate)
 		{
-			ConfigValues.Emplace(TEXT("DefaultEngine.ini"),
-				TEXT("/Script/WorldPartitionEditor.WorldPartitionEditorSettings"),
-				TEXT("bEnableConversionPrompt"),
-				TEXT("False"),
-				true /* ShouldReplaceExistingValue */);
-
 			ConfigValues.Emplace(TEXT("DefaultEngine.ini"),
 				TEXT("/Script/WorldPartitionEditor.WorldPartitionEditorSettings"),
 				TEXT("CommandletClass"),
diff --git a/Engine/Source/Editor/UnrealEd/Classes/Commandlets/WorldPartitionConvertCommandlet.h b/Engine/Source/Editor/UnrealEd/Classes/Commandlets/WorldPartitionConvertCommandlet.h
--- a/Engine/Source/Editor/UnrealEd/Classes/Commandlets/WorldPartitionConvertCommandlet.h
+++ b/Engine/Source/Editor/UnrealEd/Classes/Commandlets/WorldPartitionConvertCommandlet.h
@@ -41,6 +41,8 @@ public:
 	virtual int32 Main(const FString& Params) override;
 	//~ End UCommandlet Interface
 
+	const FString& GetConversionSuffix() const { return ConversionSuffix; }
+
 private:
 	void GatherAndPrepareSubLevelsToConvert(const UWorldPartition* WorldPartition, ULevel* Level, TArray<ULevel*>& SubLevels);
 	EActorGridPlacement GetLevelGridPlacement(ULevel* Level, EActorGridPlacement DefaultGridPlacement);
diff --git a/Engine/Source/Editor/UnrealEd/Private/FileHelpers.cpp b/Engine/Source/Editor/UnrealEd/Private/FileHelpers.cpp
--- a/Engine/Source/Editor/UnrealEd/Private/FileHelpers.cpp
+++ b/Engine/Source/Editor/UnrealEd/Private/FileHelpers.cpp
@@ -2735,15 +2735,6 @@ bool FEditorFileUtils::LoadMap(const FString& InFilename, bool LoadAsTemplate, b
 		return false;
 	}
 
-	if (!LoadAsTemplate)
-	{
-		IWorldPartitionEditorModule& WorldPartitionEditorModule = FModuleManager::LoadModuleChecked<IWorldPartitionEditorModule>("WorldPartitionEditor");
-		if (!WorldPartitionEditorModule.ConvertMap(LongMapPackageName))
-		{
-			return false;
-		}
-	}
-
 	// Save last opened level name.
 	GConfig->SetString(TEXT("EditorStartup"), TEXT("LastLevel"), *LongMapPackageName, GEditorPerProjectIni);
 
@@ -4307,14 +4298,6 @@ void FEditorFileUtils::LoadDefaultMapAtStartup()
 	const bool bIncludeReadOnlyRoots = true;
 	if ( FPackageName::IsValidLongPackageName(EditorStartupMap, bIncludeReadOnlyRoots) )
 	{
-		IWorldPartitionEditorModule& WorldPartitionEditorModule = FModuleManager::LoadModuleChecked<IWorldPartitionEditorModule>("WorldPartitionEditor");
-		bool bPreviousConversionPromptEnabled = WorldPartitionEditorModule.IsConversionPromptEnabled();
-		WorldPartitionEditorModule.SetConversionPromptEnabled(false);
-		ON_SCOPE_EXIT
-		{
-			WorldPartitionEditorModule.SetConversionPromptEnabled(bPreviousConversionPromptEnabled);
-		};
-		
 		FString MapFilenameToLoad = FPackageName::LongPackageNameToFilename( EditorStartupMap );
 
 		bIsLoadingDefaultStartupMap = true;
diff --git a/Engine/Source/Editor/WorldPartitionEditor/Private/WorldPartitionEditorModule.cpp b/Engine/Source/Editor/WorldPartitionEditor/Private/WorldPartitionEditorModule.cpp
--- a/Engine/Source/Editor/WorldPartitionEditor/Private/WorldPartitionEditorModule.cpp
+++ b/Engine/Source/Editor/WorldPartitionEditor/Private/WorldPartitionEditorModule.cpp
@@ -17,13 +17,20 @@
 #include "Framework/MultiBox/MultiBoxBuilder.h"
 
 #include "Engine/Level.h"
+
 #include "Misc/MessageDialog.h"
 #include "Misc/StringBuilder.h"
 #include "Misc/ScopedSlowTask.h"
+
 #include "Interfaces/IMainFrameModule.h"
 #include "Widgets/SWindow.h"
 #include "Commandlets/WorldPartitionConvertCommandlet.h"
 #include "FileHelpers.h"
+#include "ToolMenus.h"
+#include "IContentBrowserSingleton.h"
+#include "ContentBrowserModule.h"
+#include "EditorDirectories.h"
+#include "AssetRegistryModule.h"
 
 IMPLEMENT_MODULE( FWorldPartitionEditorModule, WorldPartitionEditor );
 
@@ -91,6 +98,17 @@ void FWorldPartitionEditorModule::StartupModule()
 
 		MenuExtenderDelegates.Add(FLevelEditorModule::FLevelViewportMenuExtender_SelectedActors::CreateStatic(&OnExtendLevelEditorMenu));
 		LevelEditorExtenderDelegateHandle = MenuExtenderDelegates.Last().GetHandle();
+
+		FToolMenuOwnerScoped OwnerScoped(this);
+		UToolMenu* Menu = UToolMenus::Get()->ExtendMenu("LevelEditor.MainMenu.Tools");
+		FToolMenuSection& Section = Menu->AddSection("WorldPartition", LOCTEXT("WorldPartition", "World Partition"));
+		Section.AddEntry(FToolMenuEntry::InitMenuEntry(
+			"WorldPartition",
+			LOCTEXT("WorldPartitionConvertTitle", "Convert Level..."),
+			LOCTEXT("WorldPartitionConvertTooltip", "Converts a Level to World Partition."),
+			FSlateIcon(FEditorStyle::GetStyleSetName(), "DeveloperTools.MenuIcon"),
+			FUIAction(FExecuteAction::CreateRaw(this, &FWorldPartitionEditorModule::OnConvertMap))
+		));
 	}
 
 	IAssetTools& AssetTools = FModuleManager::LoadModuleChecked<FAssetToolsModule>("AssetTools").Get();
@@ -106,6 +124,8 @@ void FWorldPartitionEditorModule::ShutdownModule()
 		{
 			LevelEditorModule->GetAllLevelViewportContextMenuExtenders().RemoveAll([=](const FLevelEditorModule::FLevelViewportMenuExtender_SelectedActors& In) { return In.GetHandle() == LevelEditorExtenderDelegateHandle; });
 		}
+
+		UToolMenus::UnregisterOwner(this);
 	}
 
 	IAssetTools* AssetToolsModule = nullptr;
@@ -131,37 +151,37 @@ TSharedRef<SWidget> FWorldPartitionEditorModule::CreateWorldPartitionEditor()
 	return SNew(SWorldPartitionEditor).InWorld(EditorWorld);
 }
 
-bool FWorldPartitionEditorModule::IsConversionPromptEnabled() const
-{
-	return GetDefault<UWorldPartitionEditorSettings>()->bEnableConversionPrompt;
-}
-
-void FWorldPartitionEditorModule::SetConversionPromptEnabled(bool bEnabled)
-{
-	GetMutableDefault<UWorldPartitionEditorSettings>()->bEnableConversionPrompt = bEnabled;
-}
-
 float FWorldPartitionEditorModule::GetAutoCellLoadingMaxWorldSize() const
 {
 	return GetDefault<UWorldPartitionEditorSettings>()->AutoCellLoadingMaxWorldSize;
 }
 
-bool FWorldPartitionEditorModule::ConvertMap(const FString& InLongPackageName)
+void FWorldPartitionEditorModule::OnConvertMap()
 {
-	if (!IsConversionPromptEnabled() || ULevel::GetIsLevelPartitionedFromPackage(FName(*InLongPackageName)))
+	IContentBrowserSingleton& ContentBrowserSingleton = FModuleManager::LoadModuleChecked<FContentBrowserModule>("ContentBrowser").Get();
+	
+	FOpenAssetDialogConfig Config;
+	Config.bAllowMultipleSelection = false;
+	FString OutPathName;
+	if (FPackageName::TryConvertFilenameToLongPackageName(FEditorDirectories::Get().GetLastDirectory(ELastDirectory::LEVEL), OutPathName))
 	{
-		return true;
-	}
+		Config.DefaultPath = OutPathName;
+	}	
+	Config.AssetClassNames.Add(UWorld::StaticClass()->GetFName());
 
-	FText Title = LOCTEXT("ConvertMapMsgTitle", "Opening non World Partition map");
-	EAppReturnType::Type Ret = FMessageDialog::Open(EAppMsgType::YesNoCancel, LOCTEXT("ConvertMapMsg", "Do you want to convert this map? (Process will restart editor)"), &Title);
-	if (Ret == EAppReturnType::No)
+	TArray<FAssetData> Assets = ContentBrowserSingleton.CreateModalOpenAssetDialog(Config);
+	if (Assets.Num() == 1)
 	{
-		return true;
+		ConvertMap(Assets[0].PackageName.ToString());
 	}
-	else if (Ret == EAppReturnType::Cancel)
+}
+
+bool FWorldPartitionEditorModule::ConvertMap(const FString& InLongPackageName)
+{
+	if (ULevel::GetIsLevelPartitionedFromPackage(FName(*InLongPackageName)))
 	{
-		return false;
+		FMessageDialog::Open(EAppMsgType::Ok, LOCTEXT("ConvertMapMsg", "Map is already using World Partition"));
+		return true;
 	}
 
 	UWorldPartitionConvertOptions* DefaultConvertOptions = GetMutableDefault<UWorldPartitionConvertOptions>();
@@ -191,7 +211,13 @@ bool FWorldPartitionEditorModule::ConvertMap(const FString& InLongPackageName)
 
 	if (ConvertDialog->ClickedOk())
 	{
-		// User will already have been prompted to save on file open
+		// Conversion will try to load the converted map so ask user to save dirty packages
+		if (!FEditorFileUtils::SaveDirtyPackages(/*bPromptUserToSave=*/true, /*bSaveMapPackages=*/true, /*bSaveContentPackages=*/false))
+		{
+			return false;
+		}
+
+		// Unload any loaded map
 		if (!UEditorLoadingAndSavingUtils::NewBlankMap(/*bSaveExistingMap*/false))
 		{
 			return false;
@@ -233,7 +259,7 @@ bool FWorldPartitionEditorModule::ConvertMap(const FString& InLongPackageName)
 		{	
 			if (Result == 0)
 			{
-				FMessageDialog::Open(EAppMsgType::Ok, LOCTEXT("ConvertMapCompleted", "Conversion succeeded. Editor will now restart."));
+				FMessageDialog::Open(EAppMsgType::Ok, LOCTEXT("ConvertMapCompleted", "Conversion succeeded."));
 
 #if	PLATFORM_DESKTOP
 				if (DefaultConvertOptions->bGenerateIni)
@@ -243,7 +269,22 @@ bool FWorldPartitionEditorModule::ConvertMap(const FString& InLongPackageName)
 					FPlatformProcess::ExploreFolder(*PackageDirectory);
 				}
 #endif
-				FUnrealEdMisc::Get().RestartEditor(false);
+				
+				
+				// Force update before loading converted map
+				FAssetRegistryModule& AssetRegistryModule = FModuleManager::LoadModuleChecked<FAssetRegistryModule>("AssetRegistry");
+				IAssetRegistry& AssetRegistry = AssetRegistryModule.Get();
+												
+				FString MapToLoad = InLongPackageName;
+				if (!DefaultConvertOptions->bInPlace)
+				{
+					MapToLoad += UWorldPartitionConvertCommandlet::StaticClass()->GetDefaultObject<UWorldPartitionConvertCommandlet>()->GetConversionSuffix();
+				}
+				
+				AssetRegistry.ScanModifiedAssetFiles({ MapToLoad });
+				AssetRegistry.ScanPathsSynchronous({ ULevel::GetExternalActorsPath(MapToLoad) }, true);
+				
+				FEditorFileUtils::LoadMap(MapToLoad);
 			}
 		}
 		else if (bCancelled)
@@ -262,7 +303,6 @@ bool FWorldPartitionEditorModule::ConvertMap(const FString& InLongPackageName)
 
 UWorldPartitionEditorSettings::UWorldPartitionEditorSettings()
 {
-	bEnableConversionPrompt = false;
 	CommandletClass = UWorldPartitionConvertCommandlet::StaticClass();
 }
 
diff --git a/Engine/Source/Editor/WorldPartitionEditor/Public/WorldPartition/WorldPartitionEditorSettings.h b/Engine/Source/Editor/WorldPartitionEditor/Public/WorldPartition/WorldPartitionEditorSettings.h
--- a/Engine/Source/Editor/WorldPartitionEditor/Public/WorldPartition/WorldPartitionEditorSettings.h
+++ b/Engine/Source/Editor/WorldPartitionEditor/Public/WorldPartition/WorldPartitionEditorSettings.h
@@ -15,9 +15,6 @@ class UWorldPartitionEditorSettings : public UDeveloperSettings
 public:
 	UWorldPartitionEditorSettings();
 
-	UPROPERTY(config, EditAnywhere, Category = MapConversion, meta = (ToolTip = "This option when enabled will show a conversion prompt when opening non World Partition maps"))
-	bool bEnableConversionPrompt = false;
-
 	UPROPERTY(config, EditAnywhere, Category = MapConversion, meta = (ToolTip = "Commandlet class to use for World Parition conversion"))
 	TSubclassOf<UWorldPartitionConvertCommandlet> CommandletClass;
 
diff --git a/Engine/Source/Editor/WorldPartitionEditor/Public/WorldPartitionEditorModule.h b/Engine/Source/Editor/WorldPartitionEditor/Public/WorldPartitionEditorModule.h
--- a/Engine/Source/Editor/WorldPartitionEditor/Public/WorldPartitionEditorModule.h
+++ b/Engine/Source/Editor/WorldPartitionEditor/Public/WorldPartitionEditorModule.h
@@ -25,16 +25,6 @@ public:
 	 */
 	virtual TSharedRef<class SWidget> CreateWorldPartitionEditor();
 
-	/**
-	 *
-	 */
-	virtual bool IsConversionPromptEnabled() const override;
-
-	/**
-	 *
-	 */
-	virtual void SetConversionPromptEnabled(bool bEnabled) override;
-
 	/**
 	 *
 	 */
@@ -51,6 +41,8 @@ public:
 	virtual FWorldPartitionCreated& OnWorldPartitionCreated() override { return WorldPartitionCreatedEvent; }
 
 private:
+	void OnConvertMap();
+
 	FDelegateHandle LevelEditorExtenderDelegateHandle;
 
 	TSharedPtr<class FHLODLayerAssetTypeActions> HLODLayerAssetTypeActions;
diff --git a/Engine/Source/Editor/WorldPartitionEditor/WorldPartitionEditor.Build.cs b/Engine/Source/Editor/WorldPartitionEditor/WorldPartitionEditor.Build.cs
--- a/Engine/Source/Editor/WorldPartitionEditor/WorldPartitionEditor.Build.cs
+++ b/Engine/Source/Editor/WorldPartitionEditor/WorldPartitionEditor.Build.cs
@@ -29,7 +29,8 @@ public class WorldPartitionEditor : ModuleRules
 				"WorldBrowser",
 				"MainFrame",
 				"PropertyEditor",
-				"DeveloperSettings"
+				"DeveloperSettings",
+				"ToolMenus"
             }
 		);
 
diff --git a/Engine/Source/Runtime/Engine/Public/WorldPartition/IWorldPartitionEditorModule.h b/Engine/Source/Runtime/Engine/Public/WorldPartition/IWorldPartitionEditorModule.h
--- a/Engine/Source/Runtime/Engine/Public/WorldPartition/IWorldPartitionEditorModule.h
+++ b/Engine/Source/Runtime/Engine/Public/WorldPartition/IWorldPartitionEditorModule.h
@@ -12,8 +12,6 @@ class IWorldPartitionEditorModule : public IModuleInterface
 public:
 	virtual ~IWorldPartitionEditorModule() {}
 
-	virtual bool IsConversionPromptEnabled() const = 0;
-	virtual void SetConversionPromptEnabled(bool bEnabled) = 0;
 	virtual float GetAutoCellLoadingMaxWorldSize() const = 0;
 	virtual bool ConvertMap(const FString& InLongPackageName) = 0;
 
