Fixes updating WP Minimap UVOffset when texture switches VirtualTexture usage.

#changelist validated

diff --git a/Engine/Source/Editor/WorldPartitionEditor/Private/WorldPartition/SWorldPartitionEditorGridSpatialHash.cpp b/Engine/Source/Editor/WorldPartitionEditor/Private/WorldPartition/SWorldPartitionEditorGridSpatialHash.cpp
--- a/Engine/Source/Editor/WorldPartitionEditor/Private/WorldPartition/SWorldPartitionEditorGridSpatialHash.cpp
+++ b/Engine/Source/Editor/WorldPartitionEditor/Private/WorldPartition/SWorldPartitionEditorGridSpatialHash.cpp
@@ -33,7 +33,8 @@
 TAutoConsoleVariable<bool> CVarDebugDrawOctree(TEXT("wp.Editor.DebugDrawOctree"), false, TEXT("Whether to debug draw the World Partition octree"), ECVF_Default);
 
 SWorldPartitionEditorGridSpatialHash::SWorldPartitionEditorGridSpatialHash()
-	:WorldMiniMapBounds(ForceInit)
+	: WorldMiniMapBounds(ForceInit)
+	, WorldMinimapUsesVirtualTexture(false)
 {
 }
 
@@ -63,6 +64,13 @@ void SWorldPartitionEditorGridSpatialHash::Tick(const FGeometry& AllottedGeometr
 {
 	SWorldPartitionEditorGrid2D::Tick(AllottedGeometry, InCurrentTime, InDeltaTime);
 
+	UTexture2D* Texture2D = Cast<UTexture2D>(WorldMiniMapBrush.GetResourceObject());
+
+	if ((Texture2D != nullptr) && (Texture2D->IsCurrentlyVirtualTextured() != WorldMinimapUsesVirtualTexture))
+	{
+		UpdateWorldMiniMapDetails();
+	}
+
 	/*static double LastTime = InCurrentTime;
 	if ((InCurrentTime - LastTime) > 15.0)
 	{
@@ -104,10 +112,16 @@ void SWorldPartitionEditorGridSpatialHash::UpdateWorldMiniMapDetails()
 		WorldMiniMapBounds = FBox2D(FVector2D(WorldMiniMap->MiniMapWorldBounds.Min), FVector2D(WorldMiniMap->MiniMapWorldBounds.Max));
 		if (UTexture2D* MiniMapTexture = WorldMiniMap->MiniMapTexture)
 		{
-			if (MiniMapTexture->IsCurrentlyVirtualTextured())
+			WorldMinimapUsesVirtualTexture = MiniMapTexture->IsCurrentlyVirtualTextured();
+
+			if (WorldMinimapUsesVirtualTexture)
 			{
 				WorldMiniMapBrush.SetUVRegion(WorldMiniMap->UVOffset);
 			}
+			else
+			{
+				WorldMiniMapBrush.SetUVRegion(FBox2D({ 0.f, 0.f }, { 1.f, 1.f }));
+			}
 
 			WorldMiniMapBrush.SetImageSize(FVector2D(MiniMapTexture->GetSizeX(), MiniMapTexture->GetSizeY()));
 			WorldMiniMapBrush.SetResourceObject(MiniMapTexture);
diff --git a/Engine/Source/Editor/WorldPartitionEditor/Public/WorldPartition/SWorldPartitionEditorGridSpatialHash.h b/Engine/Source/Editor/WorldPartitionEditor/Public/WorldPartition/SWorldPartitionEditorGridSpatialHash.h
--- a/Engine/Source/Editor/WorldPartitionEditor/Public/WorldPartition/SWorldPartitionEditorGridSpatialHash.h
+++ b/Engine/Source/Editor/WorldPartitionEditor/Public/WorldPartition/SWorldPartitionEditorGridSpatialHash.h
@@ -26,4 +26,5 @@ private:
 	void UpdateWorldMiniMapDetails();
 	FBox2D          WorldMiniMapBounds;
 	FSlateBrush		WorldMiniMapBrush;
+	bool			WorldMinimapUsesVirtualTexture;
 };
