Fixes crash when generating WP Minimap using reload button with no previous texture.

#changelist validated

diff --git a/Engine/Source/Editor/WorldPartitionEditor/Private/WorldPartition/SWorldPartitionEditorGridSpatialHash.cpp b/Engine/Source/Editor/WorldPartitionEditor/Private/WorldPartition/SWorldPartitionEditorGridSpatialHash.cpp
--- a/Engine/Source/Editor/WorldPartitionEditor/Private/WorldPartition/SWorldPartitionEditorGridSpatialHash.cpp
+++ b/Engine/Source/Editor/WorldPartitionEditor/Private/WorldPartition/SWorldPartitionEditorGridSpatialHash.cpp
@@ -94,7 +94,11 @@ FReply SWorldPartitionEditorGridSpatialHash::ReloadMiniMap()
 	WorldMiniMap->Modify();
 
 	// Updating VT is not supported for now
-	WorldMiniMap->MiniMapTexture->VirtualTextureStreaming = false;
+	if (WorldMiniMap->MiniMapTexture != nullptr)
+	{
+		WorldMiniMap->MiniMapTexture->VirtualTextureStreaming = false;
+	}
+
 	WorldMiniMap->UVOffset.bIsValid = false;
 
 	FWorldPartitionMiniMapHelper::CaptureWorldMiniMapToTexture(World, WorldMiniMap, WorldMiniMap->MiniMapSize, static_cast<UTexture2D*&>(WorldMiniMap->MiniMapTexture), TEXT("MinimapTexture"), WorldMiniMap->MiniMapWorldBounds);
