Removed commented code.

diff --git a/Engine/Source/Editor/WorldPartitionEditor/Private/WorldPartition/SWorldPartitionEditorGridSpatialHash.cpp b/Engine/Source/Editor/WorldPartitionEditor/Private/WorldPartition/SWorldPartitionEditorGridSpatialHash.cpp
--- a/Engine/Source/Editor/WorldPartitionEditor/Private/WorldPartition/SWorldPartitionEditorGridSpatialHash.cpp
+++ b/Engine/Source/Editor/WorldPartitionEditor/Private/WorldPartition/SWorldPartitionEditorGridSpatialHash.cpp
@@ -70,13 +70,6 @@ void SWorldPartitionEditorGridSpatialHash::Tick(const FGeometry& AllottedGeometr
 	{
 		UpdateWorldMiniMapDetails();
 	}
-
-	/*static double LastTime = InCurrentTime;
-	if ((InCurrentTime - LastTime) > 15.0)
-	{
-		LastTime = InCurrentTime;
-		ReloadMiniMap();
-	}*/
 }
 
 FReply SWorldPartitionEditorGridSpatialHash::ReloadMiniMap()
