Prevent tab manager from summoning empty windows when no tabs can be opened

diff --git a/Engine/Source/Editor/WorldPartitionEditor/Private/WorldPartitionEditorModule.cpp b/Engine/Source/Editor/WorldPartitionEditor/Private/WorldPartitionEditorModule.cpp
--- a/Engine/Source/Editor/WorldPartitionEditor/Private/WorldPartitionEditorModule.cpp
+++ b/Engine/Source/Editor/WorldPartitionEditor/Private/WorldPartitionEditorModule.cpp
@@ -323,15 +323,21 @@ void FWorldPartitionEditorModule::OnMapChanged(uint32 MapFlags)
 {
 	if (MapFlags == MapChangeEventFlags::NewMap)
 	{
+		FLevelEditorModule* LevelEditorModule = FModuleManager::Get().GetModulePtr<FLevelEditorModule>("LevelEditor");
+	
+		TSharedPtr<FTabManager> LevelEditorTabManager = LevelEditorModule ? LevelEditorModule->GetLevelEditorTabManager() : nullptr;
+		if(LevelEditorTabManager)
+		{
+			UpdateTabPermissions(LevelEditorTabManager);
+		}
+
 		// If the world opened is a world partition world spawn the world partition tab if not open.
 		UWorld* EditorWorld = GEditor->GetEditorWorldContext().World();
 		if (EditorWorld && EditorWorld->IsPartitionedWorld())
 		{
-			FLevelEditorModule* LevelEditorModule = FModuleManager::Get().GetModulePtr<FLevelEditorModule>("LevelEditor");
-
-			if(LevelEditorModule && LevelEditorModule->GetLevelEditorTabManager() && !WorldPartitionTab.IsValid())
+			if(LevelEditorTabManager && !WorldPartitionTab.IsValid())
 			{
-				WorldPartitionTab = LevelEditorModule->GetLevelEditorTabManager()->TryInvokeTab(WorldPartitionEditorTabId, true);
+				WorldPartitionTab = LevelEditorTabManager->TryInvokeTab(WorldPartitionEditorTabId);
 			}
 		}
 		else if(TSharedPtr<SDockTab> WorldPartitionTabPin = WorldPartitionTab.Pin())
@@ -361,11 +367,27 @@ TSharedRef<SDockTab> FWorldPartitionEditorModule::SpawnWorldPartitionTab(const F
 	return NewTab;
 }
 
+void FWorldPartitionEditorModule::UpdateTabPermissions(TSharedPtr<FTabManager> InTabManager)
+{
+	UWorld* EditorWorld = GEditor->GetEditorWorldContext().World();
+	if(EditorWorld && EditorWorld->IsPartitionedWorld())
+	{
+		InTabManager->GetTabPermissionList()->RemoveDenyListItem(WorldPartitionEditorTabId, WorldPartitionEditorTabId);
+	}
+	else
+	{
+		InTabManager->GetTabPermissionList()->AddDenyListItem(WorldPartitionEditorTabId, WorldPartitionEditorTabId);
+	}
+}
+
 void FWorldPartitionEditorModule::RegisterWorldPartitionTabs(TSharedPtr<FTabManager> InTabManager)
 {
 	const IWorkspaceMenuStructure& MenuStructure = WorkspaceMenu::GetMenuStructure();
 
 	const FSlateIcon WorldPartitionIcon(FEditorStyle::GetStyleSetName(), "LevelEditor.Tabs.WorldPartition");
+
+	UpdateTabPermissions(InTabManager);
+
 	InTabManager->RegisterTabSpawner(WorldPartitionEditorTabId,
 		FOnSpawnTab::CreateRaw(this, &FWorldPartitionEditorModule::SpawnWorldPartitionTab),
 		FCanSpawnTab::CreateRaw(this, &FWorldPartitionEditorModule::CanSpawnWorldPartitionTab))
@@ -377,7 +399,7 @@ void FWorldPartitionEditorModule::RegisterWorldPartitionTabs(TSharedPtr<FTabMana
 
 void FWorldPartitionEditorModule::RegisterWorldPartitionLayout(FLayoutExtender& Extender)
 {
-	Extender.ExtendLayout(FTabId("LevelEditorSelectionDetails"), ELayoutExtensionPosition::After, FTabManager::FTab(WorldPartitionEditorTabId, ETabState::OpenedTab));
+	Extender.ExtendLayout(FTabId("LevelEditorSelectionDetails"), ELayoutExtensionPosition::After, FTabManager::FTab(WorldPartitionEditorTabId, ETabState::ClosedTab));
 }
 
 UWorldPartitionEditorSettings::UWorldPartitionEditorSettings()
diff --git a/Engine/Source/Editor/WorldPartitionEditor/Public/WorldPartitionEditorModule.h b/Engine/Source/Editor/WorldPartitionEditor/Public/WorldPartitionEditorModule.h
--- a/Engine/Source/Editor/WorldPartitionEditor/Public/WorldPartitionEditorModule.h
+++ b/Engine/Source/Editor/WorldPartitionEditor/Public/WorldPartitionEditorModule.h
@@ -56,6 +56,8 @@ private:
 
 	/** Spawns the world partition tab */
 	TSharedRef<SDockTab> SpawnWorldPartitionTab(const FSpawnTabArgs& Args);
+
+	void UpdateTabPermissions(TSharedPtr<FTabManager> InTabManager);
 private:
 	void OnConvertMap();
 
