World Partition Editor QoL: double-click on a cell will move camera there.

diff --git a/Engine/Source/Editor/WorldPartitionEditor/Private/WorldPartition/SWorldPartitionEditorGrid2D.cpp b/Engine/Source/Editor/WorldPartitionEditor/Private/WorldPartition/SWorldPartitionEditorGrid2D.cpp
--- a/Engine/Source/Editor/WorldPartitionEditor/Private/WorldPartition/SWorldPartitionEditorGrid2D.cpp
+++ b/Engine/Source/Editor/WorldPartitionEditor/Private/WorldPartition/SWorldPartitionEditorGrid2D.cpp
@@ -265,6 +265,12 @@ FReply SWorldPartitionEditorGrid2D::OnMouseButtonUp(const FGeometry& MyGeometry,
 	return FReply::Unhandled();
 }
 
+FReply SWorldPartitionEditorGrid2D::OnMouseButtonDoubleClick(const FGeometry& InMyGeometry, const FPointerEvent& InMouseEvent)
+{
+	MoveCameraHere();
+	return FReply::Handled();
+}
+
 FReply SWorldPartitionEditorGrid2D::OnMouseMove(const FGeometry& MyGeometry, const FPointerEvent& MouseEvent)
 {
 	const FVector2D CursorDelta = MouseEvent.GetCursorDelta();
diff --git a/Engine/Source/Editor/WorldPartitionEditor/Public/WorldPartition/SWorldPartitionEditorGrid2D.h b/Engine/Source/Editor/WorldPartitionEditor/Public/WorldPartition/SWorldPartitionEditorGrid2D.h
--- a/Engine/Source/Editor/WorldPartitionEditor/Public/WorldPartition/SWorldPartitionEditorGrid2D.h
+++ b/Engine/Source/Editor/WorldPartitionEditor/Public/WorldPartition/SWorldPartitionEditorGrid2D.h
@@ -40,6 +40,7 @@ public:
 
 	virtual FReply OnMouseButtonDown(const FGeometry& MyGeometry, const FPointerEvent& MouseEvent) override;
 	virtual FReply OnMouseButtonUp(const FGeometry& MyGeometry, const FPointerEvent& MouseEvent) override;
+	virtual FReply OnMouseButtonDoubleClick(const FGeometry& InMyGeometry, const FPointerEvent& InMouseEvent) override;
 	virtual FReply OnMouseMove(const FGeometry& MyGeometry, const FPointerEvent& MouseEvent) override;
 	virtual FReply OnMouseWheel(const FGeometry& MyGeometry, const FPointerEvent& MouseEvent) override;
 	virtual FCursorReply OnCursorQuery(const FGeometry& MyGeometry, const FPointerEvent& CursorEvent) const override;
