World Partition Editor QoL: added right-click option to unload all loaded cells.

diff --git a/Engine/Source/Editor/WorldPartitionEditor/Private/WorldPartition/SWorldPartitionEditorGrid2D.cpp b/Engine/Source/Editor/WorldPartitionEditor/Private/WorldPartition/SWorldPartitionEditorGrid2D.cpp
--- a/Engine/Source/Editor/WorldPartitionEditor/Private/WorldPartition/SWorldPartitionEditorGrid2D.cpp
+++ b/Engine/Source/Editor/WorldPartitionEditor/Private/WorldPartition/SWorldPartitionEditorGrid2D.cpp
@@ -54,6 +54,7 @@ void SWorldPartitionEditorGrid2D::FEditorCommands::RegisterCommands()
 {
 	UI_COMMAND(LoadSelectedCells, "Load Selected Cells", "Load the selected cells.", EUserInterfaceActionType::Button, FInputChord());
 	UI_COMMAND(UnloadSelectedCells, "Unload Selected Cells", "Unload the selected cells.", EUserInterfaceActionType::Button, FInputChord());
+	UI_COMMAND(UnloadAllCells, "Unload All Cells", "Unload all cells.", EUserInterfaceActionType::Button, FInputChord());
 	UI_COMMAND(MoveCameraHere, "Move Camera Here", "MOve the camera to the selected position.", EUserInterfaceActionType::Button, FInputChord());
 }
 
@@ -155,6 +156,7 @@ void SWorldPartitionEditorGrid2D::Construct(const FArguments& InArgs)
 
 	ActionList.MapAction(Commands.LoadSelectedCells, FExecuteAction::CreateSP(this, &SWorldPartitionEditorGrid2D::LoadSelectedCells), FCanExecuteAction::CreateLambda(CanLoadOrUnloadCells));
 	ActionList.MapAction(Commands.UnloadSelectedCells, FExecuteAction::CreateSP(this, &SWorldPartitionEditorGrid2D::UnloadSelectedCells), FCanExecuteAction::CreateLambda(CanLoadOrUnloadCells));
+	ActionList.MapAction(Commands.UnloadAllCells, FExecuteAction::CreateSP(this, &SWorldPartitionEditorGrid2D::UnloadAllCells));
 	ActionList.MapAction(Commands.MoveCameraHere, FExecuteAction::CreateSP(this, &SWorldPartitionEditorGrid2D::MoveCameraHere));
 }
 
@@ -172,6 +174,14 @@ void SWorldPartitionEditorGrid2D::UnloadSelectedCells()
 	Refresh();
 }
 
+void SWorldPartitionEditorGrid2D::UnloadAllCells()
+{
+	const FBox AllCellsBox(FVector(-WORLD_MAX, -WORLD_MAX, -WORLD_MAX), FVector(WORLD_MAX, WORLD_MAX, WORLD_MAX));
+	WorldPartition->UnloadEditorCells(AllCellsBox, true);
+	GEditor->RedrawLevelEditingViewports();
+	Refresh();
+}
+
 void SWorldPartitionEditorGrid2D::MoveCameraHere()
 {
 	for (FLevelEditorViewportClient* LevelVC : GEditor->GetLevelViewportClients())
@@ -234,6 +244,7 @@ FReply SWorldPartitionEditorGrid2D::OnMouseButtonUp(const FGeometry& MyGeometry,
 			{
 				MenuBuilder.AddMenuEntry(Commands.LoadSelectedCells);
 				MenuBuilder.AddMenuEntry(Commands.UnloadSelectedCells);
+				MenuBuilder.AddMenuEntry(Commands.UnloadAllCells);
 				MenuBuilder.AddMenuEntry(Commands.MoveCameraHere);
 			}
 			MenuBuilder.EndSection();
diff --git a/Engine/Source/Editor/WorldPartitionEditor/Public/WorldPartition/SWorldPartitionEditorGrid2D.h b/Engine/Source/Editor/WorldPartitionEditor/Public/WorldPartition/SWorldPartitionEditorGrid2D.h
--- a/Engine/Source/Editor/WorldPartitionEditor/Public/WorldPartition/SWorldPartitionEditorGrid2D.h
+++ b/Engine/Source/Editor/WorldPartitionEditor/Public/WorldPartition/SWorldPartitionEditorGrid2D.h
@@ -20,6 +20,7 @@ protected:
 	
 		TSharedPtr<FUICommandInfo> LoadSelectedCells;
 		TSharedPtr<FUICommandInfo> UnloadSelectedCells;
+		TSharedPtr<FUICommandInfo> UnloadAllCells;
 		TSharedPtr<FUICommandInfo> MoveCameraHere;
 
 		/**
@@ -36,6 +37,7 @@ public:
 
 	void LoadSelectedCells();
 	void UnloadSelectedCells();
+	void UnloadAllCells();
 	void MoveCameraHere();
 
 	virtual FReply OnMouseButtonDown(const FGeometry& MyGeometry, const FPointerEvent& MouseEvent) override;
