World Partition Editor QoL: automatically load surrounding editor cells when using the BugItGo command.
Disabled by default, can be enabled by setting the value in the World Partition Editor (saved in per-user settings).
Loading range can be tweaked via wp.Editor.LoadingRangeBugItGo cvar (defaults to 128M).

diff --git a/Engine/Source/Editor/UnrealEd/Classes/Editor/EditorEngine.h b/Engine/Source/Editor/UnrealEd/Classes/Editor/EditorEngine.h
--- a/Engine/Source/Editor/UnrealEd/Classes/Editor/EditorEngine.h
+++ b/Engine/Source/Editor/UnrealEd/Classes/Editor/EditorEngine.h
@@ -531,6 +531,10 @@ public:
 	DECLARE_MULTICAST_DELEGATE(FPreviewPlatformChanged);
 	FPreviewPlatformChanged PreviewPlatformChanged;
 
+	/** A delegate that is called when the bugitgo command is used. */
+	DECLARE_MULTICAST_DELEGATE_TwoParams(FPostBugItGoCalled, const FVector& Loc, const FRotator& Rot);
+	FPostBugItGoCalled PostBugItGoCalled;
+
 	/** Whether or not the editor is currently compiling */
 	bool bIsCompiling;
 
@@ -3211,6 +3215,9 @@ public:
 	/** Return the delegate that is called when the preview platform changes */
 	FPreviewPlatformChanged& OnPreviewPlatformChanged() { return PreviewPlatformChanged; }
 
+	/** Return the delegate that is called when the bugitgo command is called */
+	FPostBugItGoCalled& OnPostBugItGoCalled() { return PostBugItGoCalled; }
+
 protected:
 
 	/** Function pair used to save and restore the global feature level */
diff --git a/Engine/Source/Editor/UnrealEd/Private/EditorServer.cpp b/Engine/Source/Editor/UnrealEd/Private/EditorServer.cpp
--- a/Engine/Source/Editor/UnrealEd/Private/EditorServer.cpp
+++ b/Engine/Source/Editor/UnrealEd/Private/EditorServer.cpp
@@ -6461,6 +6461,14 @@ bool UEditorEngine::HandleBugItGoCommand( const TCHAR* Str, FOutputDevice& Ar )
 		}
 	}
 
+	if (GetLevelViewportClients().Num())
+	{
+		// Now that all viewports are at the same position, broadcast the change
+		Loc = GetLevelViewportClients()[0]->GetViewLocation();
+		Rot = GetLevelViewportClients()[0]->GetViewRotation();
+		PostBugItGoCalled.Broadcast(Loc, Rot);
+	}
+
 	RedrawLevelEditingViewports();
 
 	return true;
diff --git a/Engine/Source/Editor/WorldPartitionEditor/Private/WorldPartition/SWorldPartitionEditorGrid2D.cpp b/Engine/Source/Editor/WorldPartitionEditor/Private/WorldPartition/SWorldPartitionEditorGrid2D.cpp
--- a/Engine/Source/Editor/WorldPartitionEditor/Private/WorldPartition/SWorldPartitionEditorGrid2D.cpp
+++ b/Engine/Source/Editor/WorldPartitionEditor/Private/WorldPartition/SWorldPartitionEditorGrid2D.cpp
@@ -10,6 +10,7 @@
 #include "WorldPartition/WorldPartitionEditorCell.h"
 #include "WorldPartition/WorldPartitionActorDesc.h"
 #include "WorldPartition/WorldPartitionActorDescView.h"
+#include "WorldPartition/WorldPartitionEditorPerProjectUserSettings.h"
 #include "Editor/GroupActor.h"
 #include "Editor/EditorEngine.h"
 #include "Widgets/Input/SCheckBox.h"
@@ -55,7 +56,7 @@ void SWorldPartitionEditorGrid2D::FEditorCommands::RegisterCommands()
 	UI_COMMAND(LoadSelectedCells, "Load Selected Cells", "Load the selected cells.", EUserInterfaceActionType::Button, FInputChord());
 	UI_COMMAND(UnloadSelectedCells, "Unload Selected Cells", "Unload the selected cells.", EUserInterfaceActionType::Button, FInputChord());
 	UI_COMMAND(UnloadAllCells, "Unload All Cells", "Unload all cells.", EUserInterfaceActionType::Button, FInputChord());
-	UI_COMMAND(MoveCameraHere, "Move Camera Here", "MOve the camera to the selected position.", EUserInterfaceActionType::Button, FInputChord());
+	UI_COMMAND(MoveCameraHere, "Move Camera Here", "Move the camera to the selected position.", EUserInterfaceActionType::Button, FInputChord());
 }
 
 SWorldPartitionEditorGrid2D::SWorldPartitionEditorGrid2D()
@@ -125,6 +126,23 @@ void SWorldPartitionEditorGrid2D::Construct(const FArguments& InArgs)
 					]
 					+SHorizontalBox::Slot()
 					.AutoWidth()
+					[
+						SNew(SCheckBox)
+						.IsChecked(GetMutableDefault<UWorldPartitionEditorPerProjectUserSettings>()->GetBugItGoLoadCells() ? ECheckBoxState::Checked : ECheckBoxState::Unchecked)
+						.IsEnabled(true)
+						.OnCheckStateChanged(FOnCheckStateChanged::CreateLambda([=](ECheckBoxState State) { GetMutableDefault<UWorldPartitionEditorPerProjectUserSettings>()->SetBugItGoLoadCells(State == ECheckBoxState::Checked); }))
+					]
+					+SHorizontalBox::Slot()
+					.FillWidth(1.0f)
+					.VAlign(VAlign_Center)
+					[
+						SNew(STextBlock)
+						.AutoWrapText(true)
+						.IsEnabled(true)
+						.Text(LOCTEXT("BugItGoLoadCells", "BugItGo Load Cells"))
+					]
+					+SHorizontalBox::Slot()
+					.AutoWidth()
 					[
 						SNew(SButton)
 						.Text(LOCTEXT("ReloadMiniMap", "Reload MiniMap"))
diff --git a/Engine/Source/Editor/WorldPartitionEditor/Public/WorldPartition/SWorldPartitionEditorGrid2D.h b/Engine/Source/Editor/WorldPartitionEditor/Public/WorldPartition/SWorldPartitionEditorGrid2D.h
--- a/Engine/Source/Editor/WorldPartitionEditor/Public/WorldPartition/SWorldPartitionEditorGrid2D.h
+++ b/Engine/Source/Editor/WorldPartitionEditor/Public/WorldPartition/SWorldPartitionEditorGrid2D.h
@@ -88,8 +88,6 @@ protected:
 	bool bIsSelecting;
 	bool bIsDragging;
 	bool bShowActors;
-	bool bShowActorClusters;
-	bool bShowGroupActors;
 	FVector2D MouseCursorPos;
 	FVector2D MouseCursorPosWorld;
 	FVector2D LastMouseCursorPosWorldDrag;
diff --git a/Engine/Source/Editor/WorldPartitionEditor/Public/WorldPartition/SWorldPartitionEditorGridSpatialHash.h b/Engine/Source/Editor/WorldPartitionEditor/Public/WorldPartition/SWorldPartitionEditorGridSpatialHash.h
--- a/Engine/Source/Editor/WorldPartitionEditor/Public/WorldPartition/SWorldPartitionEditorGridSpatialHash.h
+++ b/Engine/Source/Editor/WorldPartitionEditor/Public/WorldPartition/SWorldPartitionEditorGridSpatialHash.h
@@ -16,8 +16,6 @@ public:
 
 	void Construct(const FArguments& InArgs);
 
-	virtual void Tick(const FGeometry& AllottedGeometry, const double InCurrentTime, const float InDeltaTime) override;
-	
 	int32 PaintGrid(const FGeometry& AllottedGeometry, const FSlateRect& MyCullingRect, FSlateWindowElementList& OutDrawElements, int32 LayerId) const override;
 
 	virtual FReply ReloadMiniMap() override;
diff --git a/Engine/Source/Runtime/Engine/Private/WorldPartition/WorldPartition.cpp b/Engine/Source/Runtime/Engine/Private/WorldPartition/WorldPartition.cpp
--- a/Engine/Source/Runtime/Engine/Private/WorldPartition/WorldPartition.cpp
+++ b/Engine/Source/Runtime/Engine/Private/WorldPartition/WorldPartition.cpp
@@ -41,6 +41,14 @@
 #include "WorldPartition/WorldPartitionActorDescViewProxy.h"
 #include "Modules/ModuleManager.h"
 #include "GameDelegates.h"
+
+static int32 GLoadingRangeBugItGo = 12800;
+static FAutoConsoleVariableRef CVarGDefaultLoadingRangeHLOD0(
+	TEXT("wp.Editor.LoadingRangeBugItGo"),
+	GLoadingRangeBugItGo,
+	TEXT("Loading range for BugItGo command."),
+	ECVF_Default
+);
 #endif //WITH_EDITOR
 
 #define LOCTEXT_NAMESPACE "WorldPartition"
@@ -589,6 +597,18 @@ bool UWorldPartition::IsMainWorldPartition() const
 	return World == GetTypedOuter<UWorld>();
 }
 
+void UWorldPartition::OnPostBugItGoCalled(const FVector& Loc, const FRotator& Rot)
+{
+#if WITH_EDITOR
+	if (GetMutableDefault<UWorldPartitionEditorPerProjectUserSettings>()->GetBugItGoLoadCells())
+	{
+		const FVector LoadExtent(GLoadingRangeBugItGo, GLoadingRangeBugItGo, WORLD_MAX);
+		const FBox LoadCellsBox(Loc - LoadExtent, Loc + LoadExtent);
+		LoadEditorCells(LoadCellsBox, false);
+	}
+#endif
+}
+
 void UWorldPartition::RegisterDelegates()
 {
 	check(World); 
@@ -602,6 +622,8 @@ void UWorldPartition::RegisterDelegates()
 		FGameDelegates::Get().GetEndPlayMapDelegate().AddUObject(this, &UWorldPartition::OnEndPlay);
 
 		FCoreUObjectDelegates::PostReachabilityAnalysis.AddUObject(this, &UWorldPartition::OnGCPostReachabilityAnalysis);
+
+		GEditor->OnPostBugItGoCalled().AddUObject(this, &UWorldPartition::OnPostBugItGoCalled);
 	}
 #endif
 
@@ -631,6 +653,8 @@ void UWorldPartition::UnregisterDelegates()
 		{
 			FCoreUObjectDelegates::PostReachabilityAnalysis.RemoveAll(this);
 		}
+
+		GEditor->OnPostBugItGoCalled().RemoveAll(this);
 	}
 #endif
 
diff --git a/Engine/Source/Runtime/Engine/Public/WorldPartition/WorldPartition.h b/Engine/Source/Runtime/Engine/Public/WorldPartition/WorldPartition.h
--- a/Engine/Source/Runtime/Engine/Public/WorldPartition/WorldPartition.h
+++ b/Engine/Source/Runtime/Engine/Public/WorldPartition/WorldPartition.h
@@ -221,6 +221,8 @@ private:
 	bool IsMainWorldPartition() const;
 	void OnWorldMatchStarting();
 
+	void OnPostBugItGoCalled(const FVector& Loc, const FRotator& Rot);
+
 	// Delegates registration
 	void RegisterDelegates();
 	void UnregisterDelegates();	
diff --git a/Engine/Source/Runtime/Engine/Public/WorldPartition/WorldPartitionEditorPerProjectUserSettings.h b/Engine/Source/Runtime/Engine/Public/WorldPartition/WorldPartitionEditorPerProjectUserSettings.h
--- a/Engine/Source/Runtime/Engine/Public/WorldPartition/WorldPartitionEditorPerProjectUserSettings.h
+++ b/Engine/Source/Runtime/Engine/Public/WorldPartition/WorldPartitionEditorPerProjectUserSettings.h
@@ -74,8 +74,22 @@ public:
 		return !bDisableLoadingOfLastLoadedCells;
 	}
 
-	TArray<FName> GetWorldDataLayersNotLoadedInEditor(UWorld* InWorld) const;
-	TArray<FName> GetWorldDataLayersLoadedInEditor(UWorld* InWorld) const;
+	bool GetBugItGoLoadCells() const
+	{
+		return bBugItGoLoadCells;
+	}
+
+	void SetBugItGoLoadCells(bool bInBugItGoLoadCells)
+	{
+		if (bBugItGoLoadCells != bInBugItGoLoadCells)
+		{
+			bBugItGoLoadCells = bInBugItGoLoadCells;
+			SaveConfig();
+		}
+	}
+
+	const TArray<FName>& GetWorldDataLayersNotLoadedInEditor(UWorld* InWorld) const;
+	const TArray<FName>& GetWorldDataLayersLoadedInEditor(UWorld* InWorld) const;
 	
 	void SetWorldDataLayersNonDefaultEditorLoadStates(UWorld* InWorld, const TArray<FName>& InDataLayersLoadedInEditor, const TArray<FName>& InDataLayersNotLoadedInEditor);
 
@@ -118,6 +132,9 @@ private:
 	UPROPERTY(config)
 	uint32 bDisableLoadingOfLastLoadedCells : 1;
 
+	UPROPERTY(config)
+	uint32 bBugItGoLoadCells : 1;
+
 	UPROPERTY(config)
 	TMap<TSoftObjectPtr<UWorld>, FWorldPartitionPerWorldSettings> PerWorldEditorSettings;
 #endif
