World Partition Editor QoL: display editor cell coordinates.
Disabled by default, can be enabled by setting the value in the World Partition Editor (saved in per-user settings).

diff --git a/Engine/Source/Editor/WorldPartitionEditor/Private/WorldPartition/SWorldPartitionEditorGrid2D.cpp b/Engine/Source/Editor/WorldPartitionEditor/Private/WorldPartition/SWorldPartitionEditorGrid2D.cpp
--- a/Engine/Source/Editor/WorldPartitionEditor/Private/WorldPartition/SWorldPartitionEditorGrid2D.cpp
+++ b/Engine/Source/Editor/WorldPartitionEditor/Private/WorldPartition/SWorldPartitionEditorGrid2D.cpp
@@ -143,6 +143,23 @@ void SWorldPartitionEditorGrid2D::Construct(const FArguments& InArgs)
 					]
 					+SHorizontalBox::Slot()
 					.AutoWidth()
+					[
+						SNew(SCheckBox)
+						.IsChecked(GetMutableDefault<UWorldPartitionEditorPerProjectUserSettings>()->GetShowCellCoords() ? ECheckBoxState::Checked : ECheckBoxState::Unchecked)
+						.IsEnabled(true)
+						.OnCheckStateChanged(FOnCheckStateChanged::CreateLambda([=](ECheckBoxState State) { GetMutableDefault<UWorldPartitionEditorPerProjectUserSettings>()->SetShowCellCoords(State == ECheckBoxState::Checked); }))
+					]
+					+SHorizontalBox::Slot()
+					.FillWidth(1.0f)
+					.VAlign(VAlign_Center)
+					[
+						SNew(STextBlock)
+						.AutoWrapText(true)
+						.IsEnabled(true)
+						.Text(LOCTEXT("ShowCellCoords", "Show Cell Coords"))
+					]
+					+SHorizontalBox::Slot()
+					.AutoWidth()
 					[
 						SNew(SButton)
 						.Text(LOCTEXT("ReloadMiniMap", "Reload MiniMap"))
diff --git a/Engine/Source/Editor/WorldPartitionEditor/Private/WorldPartition/SWorldPartitionEditorGridSpatialHash.cpp b/Engine/Source/Editor/WorldPartitionEditor/Private/WorldPartition/SWorldPartitionEditorGridSpatialHash.cpp
--- a/Engine/Source/Editor/WorldPartitionEditor/Private/WorldPartition/SWorldPartitionEditorGridSpatialHash.cpp
+++ b/Engine/Source/Editor/WorldPartitionEditor/Private/WorldPartition/SWorldPartitionEditorGridSpatialHash.cpp
@@ -9,6 +9,7 @@
 #include "WorldPartition/WorldPartitionActorDesc.h"
 #include "WorldPartition/WorldPartitionMiniMap.h"
 #include "WorldPartition/WorldPartitionMiniMapHelper.h"
+#include "WorldPartition/WorldPartitionEditorPerProjectUserSettings.h"
 #include "Framework/Commands/Commands.h"
 #include "Framework/MultiBox/MultiBoxBuilder.h"
 #include "Misc/FileHelper.h"
@@ -25,6 +26,7 @@
 #include "Engine/Texture2D.h"
 #include "Engine/Texture2DDynamic.h"
 #include "Misc/HashBuilder.h"
+#include "Fonts/FontMeasure.h"
 #include "EngineModule.h"
 #include "Renderer/Private/RendererModule.h"
 
@@ -350,6 +352,54 @@ int32 SWorldPartitionEditorGridSpatialHash::PaintGrid(const FGeometry& AllottedG
 			FSlateDrawElement::MakeLines(OutDrawElements, LayerId, AllottedGeometry.ToPaintGeometry(), LinePoints, ESlateDrawEffect::NoBlending, Color, false, 1.0f);
 		}
 
+		// Draw coordinates
+		if (GetMutableDefault<UWorldPartitionEditorPerProjectUserSettings>()->GetShowCellCoords())
+		{
+			const TSharedRef<FSlateFontMeasure> FontMeasure = FSlateApplication::Get().GetRenderer()->GetFontMeasureService();
+			const FVector2D CellScreenSize = WorldToScreen.TransformVector(FVector2D(EditorSpatialHash->CellSize, EditorSpatialHash->CellSize));
+
+			FSlateFontInfo CoordsFont;
+			FVector2D DefaultCoordTextSize;
+			bool bNeedsGradient = true;
+			for(int32 DesiredFontSize = 24; DesiredFontSize >= 8; DesiredFontSize -= 2)
+			{
+				CoordsFont = FCoreStyle::GetDefaultFontStyle("Bold", DesiredFontSize);
+				DefaultCoordTextSize = FontMeasure->Measure(TEXT("(-99,-99)"), CoordsFont);				
+
+				if (CellScreenSize.X > DefaultCoordTextSize.X)
+				{
+					bNeedsGradient = (DesiredFontSize == 8);
+					break;
+				}
+			}
+
+			if (CellScreenSize.X > DefaultCoordTextSize.X)
+			{
+				static float GradientDistance = 64.0f;
+				float ColorGradient = bNeedsGradient ? FMath::Min((CellScreenSize.X - DefaultCoordTextSize.X) / GradientDistance, 1.0f) : 1.0f;
+				const FLinearColor CoordTextColor(1.0f, 1.0f, 1.0f, ColorGradient);
+
+				for (int32 y = TopLeftW.Y; y < BottomRightW.Y; y += EditorSpatialHash->CellSize)
+				{
+					for (int32 x = TopLeftW.X; x < BottomRightW.X; x += EditorSpatialHash->CellSize)
+					{
+						const FString CoordText(FString::Printf(TEXT("(%d,%d)"), x / EditorSpatialHash->CellSize, y / EditorSpatialHash->CellSize));
+						const FVector2D CoordTextSize = FontMeasure->Measure(CoordText, CoordsFont);
+
+						FSlateDrawElement::MakeText(
+							OutDrawElements,
+							++LayerId,
+							AllottedGeometry.ToPaintGeometry(WorldToScreen.TransformPoint(FVector2D(x + EditorSpatialHash->CellSize / 2, y + EditorSpatialHash->CellSize / 2)) - CoordTextSize / 2, FVector2D(1,1)),
+							FString::Printf(TEXT("(%d,%d)"), x / EditorSpatialHash->CellSize, y / EditorSpatialHash->CellSize),
+							CoordsFont,
+							ESlateDrawEffect::None,
+							CoordTextColor
+						);
+					}
+				}
+			}
+		}
+
 		++LayerId;
 	}
 
diff --git a/Engine/Source/Runtime/Engine/Public/WorldPartition/WorldPartitionEditorPerProjectUserSettings.h b/Engine/Source/Runtime/Engine/Public/WorldPartition/WorldPartitionEditorPerProjectUserSettings.h
--- a/Engine/Source/Runtime/Engine/Public/WorldPartition/WorldPartitionEditorPerProjectUserSettings.h
+++ b/Engine/Source/Runtime/Engine/Public/WorldPartition/WorldPartitionEditorPerProjectUserSettings.h
@@ -88,6 +88,20 @@ public:
 		}
 	}
 
+	bool GetShowCellCoords() const
+	{
+		return bShowCellCoords;
+	}
+
+	void SetShowCellCoords(bool bInShowCellCoords)
+	{
+		if (bShowCellCoords != bInShowCellCoords)
+		{
+			bShowCellCoords = bInShowCellCoords;
+			SaveConfig();
+		}
+	}
+
 	const TArray<FName>& GetWorldDataLayersNotLoadedInEditor(UWorld* InWorld) const;
 	const TArray<FName>& GetWorldDataLayersLoadedInEditor(UWorld* InWorld) const;
 	
@@ -135,6 +149,9 @@ private:
 	UPROPERTY(config)
 	uint32 bBugItGoLoadCells : 1;
 
+	UPROPERTY(config)
+	uint32 bShowCellCoords : 1;
+
 	UPROPERTY(config)
 	TMap<TSoftObjectPtr<UWorld>, FWorldPartitionPerWorldSettings> PerWorldEditorSettings;
 #endif
