Integration Fix: missing Tick declaration

diff --git a/Engine/Source/Editor/WorldPartitionEditor/Public/WorldPartition/SWorldPartitionEditorGridSpatialHash.h b/Engine/Source/Editor/WorldPartitionEditor/Public/WorldPartition/SWorldPartitionEditorGridSpatialHash.h
--- a/Engine/Source/Editor/WorldPartitionEditor/Public/WorldPartition/SWorldPartitionEditorGridSpatialHash.h
+++ b/Engine/Source/Editor/WorldPartitionEditor/Public/WorldPartition/SWorldPartitionEditorGridSpatialHash.h
@@ -16,6 +16,8 @@ public:
 
 	void Construct(const FArguments& InArgs);
 
+	virtual void Tick(const FGeometry& AllottedGeometry, const double InCurrentTime, const float InDeltaTime) override;
+
 	int32 PaintGrid(const FGeometry& AllottedGeometry, const FSlateRect& MyCullingRect, FSlateWindowElementList& OutDrawElements, int32 LayerId) const override;
 
 	virtual FReply ReloadMiniMap() override;
@@ -25,4 +27,4 @@ private:
 	FBox2D          WorldMiniMapBounds;
 	FSlateBrush		WorldMiniMapBrush;
 	bool			WorldMinimapUsesVirtualTexture;
-};
+};
\ No newline at end of file
