Show a message about world partition not being enabled if the world partition tab is opened in a users layout and do not attempt to control the tab if in a user layout.

diff --git a/Engine/Source/Editor/WorldPartitionEditor/Private/WorldPartition/SWorldPartitionEditorGrid.cpp b/Engine/Source/Editor/WorldPartitionEditor/Private/WorldPartition/SWorldPartitionEditorGrid.cpp
--- a/Engine/Source/Editor/WorldPartitionEditor/Private/WorldPartition/SWorldPartitionEditorGrid.cpp
+++ b/Engine/Source/Editor/WorldPartitionEditor/Private/WorldPartition/SWorldPartitionEditorGrid.cpp
@@ -11,11 +11,24 @@
 
 TMap<FName, SWorldPartitionEditorGrid::PartitionEditorGridCreateInstanceFunc> SWorldPartitionEditorGrid::PartitionEditorGridCreateInstanceFactory;
 
+#define LOCTEXT_NAMESPACE "WorldPartitionEditor"
+
 void SWorldPartitionEditorGrid::Construct(const FArguments& InArgs)
 {
 	World = InArgs._InWorld;
 	WorldPartition = World ? World->GetWorldPartition() : nullptr;
 
+	if (!WorldPartition)
+	{
+		ChildSlot
+		.VAlign(VAlign_Center)
+		.HAlign(HAlign_Center)
+		[
+			SNew(STextBlock)
+			.Text(LOCTEXT("WorldPartitionMustBeEnabled", "World Partition is Not Enabled for this Map"))
+			.ColorAndOpacity(FSlateColor::UseForeground())
+		];
+	}
 	SetClipping(EWidgetClipping::ClipToBounds);
 }
 
@@ -85,4 +98,6 @@ void SWorldPartitionEditorGrid::Refresh()
 			SceneOutlinerPtr->FullRefresh();
 		}
 	}
-}
\ No newline at end of file
+}
+
+#undef LOCTEXT_NAMESPACE
\ No newline at end of file
diff --git a/Engine/Source/Editor/WorldPartitionEditor/Private/WorldPartitionEditorModule.cpp b/Engine/Source/Editor/WorldPartitionEditor/Private/WorldPartitionEditorModule.cpp
--- a/Engine/Source/Editor/WorldPartitionEditor/Private/WorldPartitionEditorModule.cpp
+++ b/Engine/Source/Editor/WorldPartitionEditor/Private/WorldPartitionEditorModule.cpp
@@ -326,10 +326,6 @@ void FWorldPartitionEditorModule::OnMapChanged(uint32 MapFlags)
 		FLevelEditorModule* LevelEditorModule = FModuleManager::Get().GetModulePtr<FLevelEditorModule>("LevelEditor");
 	
 		TSharedPtr<FTabManager> LevelEditorTabManager = LevelEditorModule ? LevelEditorModule->GetLevelEditorTabManager() : nullptr;
-		if(LevelEditorTabManager)
-		{
-			UpdateTabPermissions(LevelEditorTabManager);
-		}
 
 		// If the world opened is a world partition world spawn the world partition tab if not open.
 		UWorld* EditorWorld = GEditor->GetEditorWorldContext().World();
@@ -348,12 +344,6 @@ void FWorldPartitionEditorModule::OnMapChanged(uint32 MapFlags)
 	}
 }
 
-bool FWorldPartitionEditorModule::CanSpawnWorldPartitionTab(const FSpawnTabArgs& Args)
-{
-	UWorld* EditorWorld = GEditor->GetEditorWorldContext().World();
-	return EditorWorld && EditorWorld->IsPartitionedWorld();
-}
-
 TSharedRef<SDockTab> FWorldPartitionEditorModule::SpawnWorldPartitionTab(const FSpawnTabArgs& Args)
 {
 	TSharedRef<SDockTab> NewTab =
@@ -367,30 +357,14 @@ TSharedRef<SDockTab> FWorldPartitionEditorModule::SpawnWorldPartitionTab(const F
 	return NewTab;
 }
 
-void FWorldPartitionEditorModule::UpdateTabPermissions(TSharedPtr<FTabManager> InTabManager)
-{
-	UWorld* EditorWorld = GEditor->GetEditorWorldContext().World();
-	if(EditorWorld && EditorWorld->IsPartitionedWorld())
-	{
-		InTabManager->GetTabPermissionList()->RemoveDenyListItem(WorldPartitionEditorTabId, WorldPartitionEditorTabId);
-	}
-	else
-	{
-		InTabManager->GetTabPermissionList()->AddDenyListItem(WorldPartitionEditorTabId, WorldPartitionEditorTabId);
-	}
-}
-
 void FWorldPartitionEditorModule::RegisterWorldPartitionTabs(TSharedPtr<FTabManager> InTabManager)
 {
 	const IWorkspaceMenuStructure& MenuStructure = WorkspaceMenu::GetMenuStructure();
 
 	const FSlateIcon WorldPartitionIcon(FEditorStyle::GetStyleSetName(), "LevelEditor.Tabs.WorldPartition");
 
-	UpdateTabPermissions(InTabManager);
-
 	InTabManager->RegisterTabSpawner(WorldPartitionEditorTabId,
-		FOnSpawnTab::CreateRaw(this, &FWorldPartitionEditorModule::SpawnWorldPartitionTab),
-		FCanSpawnTab::CreateRaw(this, &FWorldPartitionEditorModule::CanSpawnWorldPartitionTab))
+		FOnSpawnTab::CreateRaw(this, &FWorldPartitionEditorModule::SpawnWorldPartitionTab))
 		.SetDisplayName(NSLOCTEXT("LevelEditorTabs", "WorldPartitionEditor", "World Partition Editor"))
 		.SetTooltipText(NSLOCTEXT("LevelEditorTabs", "WorldPartitionEditorTooltipText", "Open the World Partition Editor."))
 		.SetGroup(MenuStructure.GetLevelEditorWorldPartitionCategory())
diff --git a/Engine/Source/Editor/WorldPartitionEditor/Public/WorldPartition/SWorldPartitionEditorGrid.h b/Engine/Source/Editor/WorldPartitionEditor/Public/WorldPartition/SWorldPartitionEditorGrid.h
--- a/Engine/Source/Editor/WorldPartitionEditor/Public/WorldPartition/SWorldPartitionEditorGrid.h
+++ b/Engine/Source/Editor/WorldPartitionEditor/Public/WorldPartition/SWorldPartitionEditorGrid.h
@@ -22,7 +22,7 @@ class UWorldPartition;
 /**
  * Base class for world partition editors (goes hand in hand with corresponding UWorldPartition class via GetWorldPartitionEditorName)
  */
-class SWorldPartitionEditorGrid : public SPanel
+class SWorldPartitionEditorGrid : public SCompoundWidget
 {
 	friend class SWorldPartitionEditor;
 
@@ -44,16 +44,20 @@ public:
 	// SPanel interface
 	virtual void OnArrangeChildren(const FGeometry& AllottedGeometry, FArrangedChildren& ArrangedChildren) const override
 	{
+		if (!WorldPartition)
+		{
+			SCompoundWidget::OnArrangeChildren(AllottedGeometry, ArrangedChildren);
+		}
 	}
 
-	virtual FVector2D ComputeDesiredSize(float) const override
+	virtual FVector2D ComputeDesiredSize(float Scale) const override
 	{
-		return FVector2D(100.0f, 100.0f);
+		return !WorldPartition ? FVector2D(100.0f, 100.0f) : SCompoundWidget::ComputeDesiredSize(Scale);
 	}
 
 	virtual FChildren* GetChildren() override
 	{
-		return &FNoChildren::NoChildrenInstance;
+		return !WorldPartition ? &ChildSlot : (FChildren*)&FNoChildren::NoChildrenInstance;
 	}
 	
 	bool GetPlayerView(FVector& Location, FRotator& Rotation) const;
diff --git a/Engine/Source/Editor/WorldPartitionEditor/Public/WorldPartitionEditorModule.h b/Engine/Source/Editor/WorldPartitionEditor/Public/WorldPartitionEditorModule.h
--- a/Engine/Source/Editor/WorldPartitionEditor/Public/WorldPartitionEditorModule.h
+++ b/Engine/Source/Editor/WorldPartitionEditor/Public/WorldPartitionEditorModule.h
@@ -51,13 +51,9 @@ private:
 	/** Inserts world partition tabs into the level editor layout */
 	void RegisterWorldPartitionLayout(FLayoutExtender& Extender);
 	
-	/** Determines if ia world partition tab can be spawned */	
-	bool CanSpawnWorldPartitionTab(const FSpawnTabArgs& Args);
-
 	/** Spawns the world partition tab */
 	TSharedRef<SDockTab> SpawnWorldPartitionTab(const FSpawnTabArgs& Args);
 
-	void UpdateTabPermissions(TSharedPtr<FTabManager> InTabManager);
 private:
 	void OnConvertMap();
 
