Fix non-unity

diff --git a/Engine/Source/Editor/WorldPartitionEditor/Public/WorldPartition/SWorldPartitionEditorGrid.h b/Engine/Source/Editor/WorldPartitionEditor/Public/WorldPartition/SWorldPartitionEditorGrid.h
--- a/Engine/Source/Editor/WorldPartitionEditor/Public/WorldPartition/SWorldPartitionEditorGrid.h
+++ b/Engine/Source/Editor/WorldPartitionEditor/Public/WorldPartition/SWorldPartitionEditorGrid.h
@@ -4,7 +4,7 @@
 #include "CoreMinimal.h"
 #include "Layout/Children.h"
 #include "Widgets/DeclarativeSyntaxSupport.h"
-#include "Widgets/SPanel.h"
+#include "Widgets/SCompoundWidget.h"
 
 class UWorldPartition;
 
