WorldPartition now has its own max world size (value is currently equal to UE_OLD_WORLD_MAX).

diff --git a/Engine/Source/Editor/UnrealEd/Private/Commandlets/WorldPartitionConvertCommandlet.cpp b/Engine/Source/Editor/UnrealEd/Private/Commandlets/WorldPartitionConvertCommandlet.cpp
--- a/Engine/Source/Editor/UnrealEd/Private/Commandlets/WorldPartitionConvertCommandlet.cpp
+++ b/Engine/Source/Editor/UnrealEd/Private/Commandlets/WorldPartitionConvertCommandlet.cpp
@@ -160,7 +160,7 @@ UWorldPartitionConvertCommandlet::UWorldPartitionConvertCommandlet(const FObject
 	, ConversionSuffix(TEXT("_WP"))
 	, bConvertActorsNotReferencedByLevelScript(true)
 	, WorldOrigin(FVector::ZeroVector)
-	, WorldExtent(HALF_WORLD_MAX)
+	, WorldExtent(WORLDPARTITION_MAX * 0.5)
 	, LandscapeGridSize(4)
 {}
 
diff --git a/Engine/Source/Editor/UnrealEd/Private/WorldPartition/WorldPartitionBuilder.cpp b/Engine/Source/Editor/UnrealEd/Private/WorldPartition/WorldPartitionBuilder.cpp
--- a/Engine/Source/Editor/UnrealEd/Private/WorldPartition/WorldPartitionBuilder.cpp
+++ b/Engine/Source/Editor/UnrealEd/Private/WorldPartition/WorldPartitionBuilder.cpp
@@ -242,7 +242,7 @@ bool UWorldPartitionBuilder::Run(UWorld* World, FPackageSourceControlHelper& Pac
 		FBox BoundsToLoad(ForceInit);
 		if (LoadingMode == ELoadingMode::EntireWorld)
 		{
-			BoundsToLoad += FBox(FVector(-WORLD_MAX, -WORLD_MAX, -WORLD_MAX), FVector(WORLD_MAX, WORLD_MAX, WORLD_MAX));
+			BoundsToLoad += FBox(FVector(-WORLDPARTITION_MAX, -WORLDPARTITION_MAX, -WORLDPARTITION_MAX), FVector(WORLDPARTITION_MAX, WORLDPARTITION_MAX, WORLDPARTITION_MAX));
 			WorldPartition->LoadEditorCells(BoundsToLoad, false);
 		}
 
diff --git a/Engine/Source/Editor/UnrealEd/Private/WorldPartition/WorldPartitionNavigationDataBuilder.cpp b/Engine/Source/Editor/UnrealEd/Private/WorldPartition/WorldPartitionNavigationDataBuilder.cpp
--- a/Engine/Source/Editor/UnrealEd/Private/WorldPartition/WorldPartitionNavigationDataBuilder.cpp
+++ b/Engine/Source/Editor/UnrealEd/Private/WorldPartition/WorldPartitionNavigationDataBuilder.cpp
@@ -307,7 +307,7 @@ bool UWorldPartitionNavigationDataBuilder::GenerateNavigationData(UWorldPartitio
 				return false;
 			}
 
-			constexpr float HalfHeight = HALF_WORLD_MAX;
+			constexpr float HalfHeight = WORLDPARTITION_MAX * 0.5;
 			const FBox QueryBounds(FVector(CellBounds.Min.X, CellBounds.Min.Y, -HalfHeight), FVector(CellBounds.Max.X, CellBounds.Max.Y, HalfHeight));
 
 			if (!NavDataBounds.IsValid || !NavDataBounds.Intersect(QueryBounds))
diff --git a/Engine/Source/Editor/WorldPartitionEditor/Private/WorldPartition/SWorldPartitionEditorGrid2D.cpp b/Engine/Source/Editor/WorldPartitionEditor/Private/WorldPartition/SWorldPartitionEditorGrid2D.cpp
--- a/Engine/Source/Editor/WorldPartitionEditor/Private/WorldPartition/SWorldPartitionEditorGrid2D.cpp
+++ b/Engine/Source/Editor/WorldPartitionEditor/Private/WorldPartition/SWorldPartitionEditorGrid2D.cpp
@@ -211,7 +211,7 @@ void SWorldPartitionEditorGrid2D::UnloadSelectedCells()
 
 void SWorldPartitionEditorGrid2D::UnloadAllCells()
 {
-	const FBox AllCellsBox(FVector(-WORLD_MAX, -WORLD_MAX, -WORLD_MAX), FVector(WORLD_MAX, WORLD_MAX, WORLD_MAX));
+	const FBox AllCellsBox(FVector(-WORLDPARTITION_MAX, -WORLDPARTITION_MAX, -WORLDPARTITION_MAX), FVector(WORLDPARTITION_MAX, WORLDPARTITION_MAX, WORLDPARTITION_MAX));
 	WorldPartition->UnloadEditorCells(AllCellsBox, true);
 	GEditor->RedrawLevelEditingViewports();
 	Refresh();
@@ -413,7 +413,7 @@ uint32 SWorldPartitionEditorGrid2D::PaintActors(const FGeometry& AllottedGeometr
 {
 	const FBox2D ViewRect(FVector2D(ForceInitToZero), AllottedGeometry.GetLocalSize());
 	const FBox2D WorldViewRect(ScreenToWorld.TransformPoint(ViewRect.Min), ScreenToWorld.TransformPoint(ViewRect.Max));
-	const FBox ViewRectWorld(FVector(WorldViewRect.Min.X, WorldViewRect.Min.Y, -WORLD_MAX), FVector(WorldViewRect.Max.X, WorldViewRect.Max.Y, WORLD_MAX));
+	const FBox ViewRectWorld(FVector(WorldViewRect.Min.X, WorldViewRect.Min.Y, -WORLDPARTITION_MAX), FVector(WorldViewRect.Max.X, WorldViewRect.Max.Y, WORLDPARTITION_MAX));
 
 	TSet<FWorldPartitionActorDescViewBoundsProxy> ActorDescList;
 
@@ -785,8 +785,8 @@ void SWorldPartitionEditorGrid2D::UpdateSelection()
 	SelectBox.Init();
 
 	const FBox SelectionBox(
-		FVector(FMath::Min(SelectionStart.X, SelectionEnd.X), FMath::Min(SelectionStart.Y, SelectionEnd.Y), -WORLD_MAX),
-		FVector(FMath::Max(SelectionStart.X, SelectionEnd.X), FMath::Max(SelectionStart.Y, SelectionEnd.Y), WORLD_MAX)
+		FVector(FMath::Min(SelectionStart.X, SelectionEnd.X), FMath::Min(SelectionStart.Y, SelectionEnd.Y), -WORLDPARTITION_MAX),
+		FVector(FMath::Max(SelectionStart.X, SelectionEnd.X), FMath::Max(SelectionStart.Y, SelectionEnd.Y), WORLDPARTITION_MAX)
 	);
 
 	SelectBox = SelectionBox;
diff --git a/Engine/Source/Runtime/Engine/Private/WorldPartition/WorldPartition.cpp b/Engine/Source/Runtime/Engine/Private/WorldPartition/WorldPartition.cpp
--- a/Engine/Source/Runtime/Engine/Private/WorldPartition/WorldPartition.cpp
+++ b/Engine/Source/Runtime/Engine/Private/WorldPartition/WorldPartition.cpp
@@ -576,7 +576,7 @@ void UWorldPartition::OnPostBugItGoCalled(const FVector& Loc, const FRotator& Ro
 #if WITH_EDITOR
 	if (GetMutableDefault<UWorldPartitionEditorPerProjectUserSettings>()->GetBugItGoLoadCells())
 	{
-		const FVector LoadExtent(GLoadingRangeBugItGo, GLoadingRangeBugItGo, WORLD_MAX);
+		const FVector LoadExtent(GLoadingRangeBugItGo, GLoadingRangeBugItGo, WORLDPARTITION_MAX);
 		const FBox LoadCellsBox(Loc - LoadExtent, Loc + LoadExtent);
 		LoadEditorCells(LoadCellsBox, false);
 	}
diff --git a/Engine/Source/Runtime/Engine/Private/WorldPartition/WorldPartitionMiniMapHelper.cpp b/Engine/Source/Runtime/Engine/Private/WorldPartition/WorldPartitionMiniMapHelper.cpp
--- a/Engine/Source/Runtime/Engine/Private/WorldPartition/WorldPartitionMiniMapHelper.cpp
+++ b/Engine/Source/Runtime/Engine/Private/WorldPartition/WorldPartitionMiniMapHelper.cpp
@@ -134,7 +134,7 @@ void FWorldPartitionMiniMapHelper::CalTopViewOfWorld(FMatrix& OutProjectionMatri
 	const float OrthoWidth = FMath::Max(1.f, ViewportWidth * OrthoZoom);
 	const float OrthoHeight = FMath::Max(1.f, ViewportHeight * OrthoZoom);
 
-	const float ZOffset = HALF_WORLD_MAX;
+	const float ZOffset = WORLDPARTITION_MAX * 0.5;
 	OutProjectionMatrix = FReversedZOrthoMatrix(
 		OrthoWidth,
 		OrthoHeight,
diff --git a/Engine/Source/Runtime/Engine/Public/WorldPartition/WorldPartition.h b/Engine/Source/Runtime/Engine/Public/WorldPartition/WorldPartition.h
--- a/Engine/Source/Runtime/Engine/Public/WorldPartition/WorldPartition.h
+++ b/Engine/Source/Runtime/Engine/Public/WorldPartition/WorldPartition.h
@@ -43,6 +43,8 @@ enum class EWorldPartitionInitState
 	Uninitializing
 };
 
+#define WORLDPARTITION_MAX UE_OLD_WORLD_MAX
+
 #if WITH_EDITOR
 /**
  * Interface for the world partition editor
