WorldPartition: Add a project setting for the Foliage Grid Size. This setting gets assigned to the WorldSettings actor when the WorldPartition gets created.

diff --git a/Engine/Source/Editor/WorldPartitionEditor/Private/WorldPartitionEditorModule.cpp b/Engine/Source/Editor/WorldPartitionEditor/Private/WorldPartitionEditorModule.cpp
--- a/Engine/Source/Editor/WorldPartitionEditor/Private/WorldPartitionEditorModule.cpp
+++ b/Engine/Source/Editor/WorldPartitionEditor/Private/WorldPartitionEditorModule.cpp
@@ -174,6 +174,17 @@ TSharedRef<SWidget> FWorldPartitionEditorModule::CreateWorldPartitionEditor()
 	return SNew(SWorldPartitionEditor).InWorld(EditorWorld);
 }
 
+int32 FWorldPartitionEditorModule::GetPlacementGridSize() const
+{
+	// Currently shares setting with Foliage. Can be changed when exposed.
+	return GetDefault<UWorldPartitionEditorSettings>()->InstancedFoliageGridSize;
+}
+
+int32 FWorldPartitionEditorModule::GetInstancedFoliageGridSize() const
+{
+	return GetDefault<UWorldPartitionEditorSettings>()->InstancedFoliageGridSize;
+}
+
 void FWorldPartitionEditorModule::OnConvertMap()
 {
 	IContentBrowserSingleton& ContentBrowserSingleton = FModuleManager::LoadModuleChecked<FContentBrowserModule>("ContentBrowser").Get();
@@ -379,6 +390,7 @@ void FWorldPartitionEditorModule::RegisterWorldPartitionLayout(FLayoutExtender&
 UWorldPartitionEditorSettings::UWorldPartitionEditorSettings()
 {
 	CommandletClass = UWorldPartitionConvertCommandlet::StaticClass();
+	InstancedFoliageGridSize = 25600;
 }
 
 FString UWorldPartitionConvertOptions::ToCommandletArgs() const
diff --git a/Engine/Source/Editor/WorldPartitionEditor/Public/WorldPartition/WorldPartitionEditorSettings.h b/Engine/Source/Editor/WorldPartitionEditor/Public/WorldPartition/WorldPartitionEditorSettings.h
--- a/Engine/Source/Editor/WorldPartitionEditor/Public/WorldPartition/WorldPartitionEditorSettings.h
+++ b/Engine/Source/Editor/WorldPartitionEditor/Public/WorldPartition/WorldPartitionEditorSettings.h
@@ -17,4 +17,7 @@ public:
 
 	UPROPERTY(config, EditAnywhere, Category = MapConversion, meta = (ToolTip = "Commandlet class to use for World Parition conversion"))
 	TSubclassOf<UWorldPartitionConvertCommandlet> CommandletClass;
+
+	UPROPERTY(config, EditAnywhere, Category = Foliage, meta = (ClampMin=3200, ToolTip= "Editor grid size used for instance foliage actors in World Partition worlds"))
+	int32 InstancedFoliageGridSize;
 };
\ No newline at end of file
diff --git a/Engine/Source/Editor/WorldPartitionEditor/Public/WorldPartitionEditorModule.h b/Engine/Source/Editor/WorldPartitionEditor/Public/WorldPartitionEditorModule.h
--- a/Engine/Source/Editor/WorldPartitionEditor/Public/WorldPartitionEditorModule.h
+++ b/Engine/Source/Editor/WorldPartitionEditor/Public/WorldPartitionEditorModule.h
@@ -31,6 +31,16 @@ public:
 	 */
 	virtual TSharedRef<class SWidget> CreateWorldPartitionEditor();
 
+	/**
+	 * Returns placement grid size setting that should be assigned to new AWorldSettings actors.
+	 */
+	virtual int32 GetPlacementGridSize() const;
+
+	/**
+	 * Returns foliage grid size setting that should be assigned to new AWorldSettings actors.
+	 */
+	virtual int32 GetInstancedFoliageGridSize() const;
+
 	/**
 	 * 
 	 */
diff --git a/Engine/Source/Runtime/Engine/Classes/GameFramework/WorldSettings.h b/Engine/Source/Runtime/Engine/Classes/GameFramework/WorldSettings.h
--- a/Engine/Source/Runtime/Engine/Classes/GameFramework/WorldSettings.h
+++ b/Engine/Source/Runtime/Engine/Classes/GameFramework/WorldSettings.h
@@ -560,7 +560,7 @@ protected:
 public:
 #if WITH_EDITORONLY_DATA
 	/** Size of the grid for instanced foliage actors */
-	UPROPERTY(EditAnywhere, Category = Foliage)
+	UPROPERTY(EditAnywhere, Category = WorldPartition)
 	uint32 InstancedFoliageGridSize;
 
 	/** Size of the grid for navigation data chunk actors */
@@ -575,7 +575,7 @@ public:
 	uint32 NavigationDataBuilderLoadingCellSize;
 	
 	/** Default size of the grid for placed elements from the editor */
-	UPROPERTY(EditAnywhere, Category=Editor)
+	UPROPERTY()
 	uint32 DefaultPlacementGridSize;
 #endif
 
diff --git a/Engine/Source/Runtime/Engine/Private/WorldPartition/WorldPartition.cpp b/Engine/Source/Runtime/Engine/Private/WorldPartition/WorldPartition.cpp
--- a/Engine/Source/Runtime/Engine/Private/WorldPartition/WorldPartition.cpp
+++ b/Engine/Source/Runtime/Engine/Private/WorldPartition/WorldPartition.cpp
@@ -32,6 +32,7 @@
 #include "Misc/ScopeExit.h"
 #include "ScopedTransaction.h"
 #include "UnrealEdMisc.h"
+#include "WorldPartition/IWorldPartitionEditorModule.h"
 #include "WorldPartition/WorldPartitionLevelHelper.h"
 #include "WorldPartition/WorldPartitionLevelStreamingDynamic.h"
 #include "WorldPartition/WorldPartitionEditorHash.h"
@@ -675,6 +676,13 @@ UWorldPartition* UWorldPartition::CreateOrRepairWorldPartition(AWorldSettings* W
 		// New maps should include GridSize in name
 		WorldSettings->bIncludeGridSizeInNameForFoliageActors = true;
 		WorldSettings->bIncludeGridSizeInNameForPartitionedActors = true;
+
+		if (IWorldPartitionEditorModule* WorldPartitionEditorModulePtr = FModuleManager::GetModulePtr<IWorldPartitionEditorModule>("WorldPartitionEditor"))
+		{
+			WorldSettings->InstancedFoliageGridSize = WorldPartitionEditorModulePtr->GetInstancedFoliageGridSize();
+			WorldSettings->DefaultPlacementGridSize = WorldPartitionEditorModulePtr->GetPlacementGridSize();
+		}
+
 		WorldSettings->MarkPackageDirty();
 
 		WorldPartition->DefaultHLODLayer = nullptr;
diff --git a/Engine/Source/Runtime/Engine/Public/WorldPartition/IWorldPartitionEditorModule.h b/Engine/Source/Runtime/Engine/Public/WorldPartition/IWorldPartitionEditorModule.h
--- a/Engine/Source/Runtime/Engine/Public/WorldPartition/IWorldPartitionEditorModule.h
+++ b/Engine/Source/Runtime/Engine/Public/WorldPartition/IWorldPartitionEditorModule.h
@@ -14,6 +14,9 @@ public:
 
 	virtual bool ConvertMap(const FString& InLongPackageName) = 0;
 
+	virtual int32 GetPlacementGridSize() const = 0;
+	virtual int32 GetInstancedFoliageGridSize() const = 0;
+
 	/** Triggered when a world is added. */
 	DECLARE_EVENT_OneParam(IWorldPartitionEditorModule, FWorldPartitionCreated, UWorld*);
 
