Show output from World Partition Commandlet when ran and completed in editor.

diff --git a/Engine/Source/Editor/WorldPartitionEditor/Private/WorldPartitionEditorModule.cpp b/Engine/Source/Editor/WorldPartitionEditor/Private/WorldPartitionEditorModule.cpp
--- a/Engine/Source/Editor/WorldPartitionEditor/Private/WorldPartitionEditorModule.cpp
+++ b/Engine/Source/Editor/WorldPartitionEditor/Private/WorldPartitionEditorModule.cpp
@@ -260,6 +260,7 @@ bool FWorldPartitionEditorModule::ConvertMap(const FString& InLongPackageName)
 		}
 
 		FProcHandle ProcessHandle;
+		FString CommandletOutput;
 		bool bCancelled = false;
 
 		// Task scope
@@ -267,6 +268,10 @@ bool FWorldPartitionEditorModule::ConvertMap(const FString& InLongPackageName)
 			FScopedSlowTask SlowTask(0, LOCTEXT("ConvertProgress", "Converting map to world partition..."));
 			SlowTask.MakeDialog(true);
 
+			void* ReadPipe = nullptr;
+			void* WritePipe = nullptr;
+			verify(FPlatformProcess::CreatePipe(ReadPipe, WritePipe));
+
 			FString CurrentExecutableName = FPlatformProcess::ExecutablePath();
 
 			// Try to provide complete Path, if we can't try with project name
@@ -274,7 +279,7 @@ bool FWorldPartitionEditorModule::ConvertMap(const FString& InLongPackageName)
 
 			uint32 ProcessID;
 			FString Arguments = FString::Printf(TEXT("\"%s\" %s"), *ProjectPath, *DefaultConvertOptions->ToCommandletArgs());
-			ProcessHandle = FPlatformProcess::CreateProc(*CurrentExecutableName, *Arguments, true, false, false, &ProcessID, 0, nullptr, nullptr);
+			ProcessHandle = FPlatformProcess::CreateProc(*CurrentExecutableName, *Arguments, true, false, false, &ProcessID, 0, nullptr, WritePipe, ReadPipe);
 			
 			while (FPlatformProcess::IsProcRunning(ProcessHandle))
 			{
@@ -285,9 +290,17 @@ bool FWorldPartitionEditorModule::ConvertMap(const FString& InLongPackageName)
 					break;
 				}
 
+				const FString LogString = FPlatformProcess::ReadPipe(ReadPipe);
+				if (!LogString.IsEmpty())
+				{
+					CommandletOutput.Append(LogString);
+				}
+
 				SlowTask.EnterProgressFrame(0);
 				FPlatformProcess::Sleep(0.1);
 			}
+
+			FPlatformProcess::ClosePipe(ReadPipe, WritePipe);
 		}
 
 		int32 Result = 0;
@@ -295,7 +308,7 @@ bool FWorldPartitionEditorModule::ConvertMap(const FString& InLongPackageName)
 		{	
 			if (Result == 0)
 			{
-				FMessageDialog::Open(EAppMsgType::Ok, LOCTEXT("ConvertMapCompleted", "Conversion succeeded."));
+				FMessageDialog::Open(EAppMsgType::Ok, FText::Format(LOCTEXT("ConvertMapCompleted", "Conversion completed:\n{0}"), FText::FromString(CommandletOutput)));
 
 #if	PLATFORM_DESKTOP
 				if (DefaultConvertOptions->bGenerateIni)
@@ -306,7 +319,6 @@ bool FWorldPartitionEditorModule::ConvertMap(const FString& InLongPackageName)
 				}
 #endif
 				
-				
 				// Force update before loading converted map
 				FAssetRegistryModule& AssetRegistryModule = FModuleManager::LoadModuleChecked<FAssetRegistryModule>("AssetRegistry");
 				IAssetRegistry& AssetRegistry = AssetRegistryModule.Get();
