Build HLODs / Minimap in editor

* Added the necessary code to be able these builds from inside the editor, using the "Build" menu
* The HLOD build will show a dialog where you can select a few options (setup, build, forcebuild, delete), then start the build
* The minimap build has no options, will start right away
* Build is actually done through an external process, so the editor will ask to save any changes that's been made, and unload the current map before starting a commandlet
* Progress is updated by parsing the commandlet log during it's execution
* Once the build is complete, will force an asset registry scan and reload the map

* Also disabled the "HLOD Outliner" tool when in a partitionned level, as it is not working and may be confusing.

diff --git a/Engine/Source/Editor/LevelEditor/Private/LevelEditor.cpp b/Engine/Source/Editor/LevelEditor/Private/LevelEditor.cpp
--- a/Engine/Source/Editor/LevelEditor/Private/LevelEditor.cpp
+++ b/Engine/Source/Editor/LevelEditor/Private/LevelEditor.cpp
@@ -1469,9 +1469,12 @@ void FLevelEditorModule::BindGlobalLevelEditorCommands()
 	ActionList.MapAction( Commands.BuildPathsOnly,
 		FExecuteAction::CreateStatic( &FLevelEditorActionCallbacks::BuildPathsOnly_Execute ) );
 
-	ActionList.MapAction(Commands.BuildLODsOnly,
-		FExecuteAction::CreateStatic(&FLevelEditorActionCallbacks::BuildLODsOnly_Execute));
+	ActionList.MapAction(Commands.BuildHLODs,
+		FExecuteAction::CreateStatic(&FLevelEditorActionCallbacks::BuildHLODs_Execute));
 	
+	ActionList.MapAction(Commands.BuildMinimap,
+		FExecuteAction::CreateStatic(&FLevelEditorActionCallbacks::BuildMinimap_Execute));
+
 	ActionList.MapAction(Commands.BuildTextureStreamingOnly,
 		FExecuteAction::CreateStatic(&FLevelEditorActionCallbacks::BuildTextureStreamingOnly_Execute));
 
diff --git a/Engine/Source/Editor/LevelEditor/Private/LevelEditorActions.cpp b/Engine/Source/Editor/LevelEditor/Private/LevelEditorActions.cpp
--- a/Engine/Source/Editor/LevelEditor/Private/LevelEditorActions.cpp
+++ b/Engine/Source/Editor/LevelEditor/Private/LevelEditorActions.cpp
@@ -799,12 +799,18 @@ void FLevelEditorActionCallbacks::BuildPathsOnly_Execute()
 	FEditorBuildUtils::EditorBuild( GetWorld(), FBuildOptions::BuildAIPaths );
 }
 
-void FLevelEditorActionCallbacks::BuildLODsOnly_Execute()
+void FLevelEditorActionCallbacks::BuildHLODs_Execute()
 {
 	// Build HLOD
 	FEditorBuildUtils::EditorBuild(GetWorld(), FBuildOptions::BuildHierarchicalLOD);
 }
 
+void FLevelEditorActionCallbacks::BuildMinimap_Execute()
+{
+	// Build HLOD
+	FEditorBuildUtils::EditorBuild(GetWorld(), FBuildOptions::BuildMinimap);
+}
+
 void FLevelEditorActionCallbacks::BuildTextureStreamingOnly_Execute()
 {
 	FEditorBuildUtils::EditorBuildTextureStreaming(GetWorld());
@@ -3380,7 +3386,8 @@ void FLevelEditorCommands::RegisterCommands()
 	UI_COMMAND( BuildGeometryOnly, "Build Geometry", "Only builds geometry (all levels.)", EUserInterfaceActionType::Button, FInputChord() );
 	UI_COMMAND( BuildGeometryOnly_OnlyCurrentLevel, "Build Geometry (Current Level)", "Builds geometry, only for the current level", EUserInterfaceActionType::Button, FInputChord() );
 	UI_COMMAND( BuildPathsOnly, "Build Paths", "Only builds paths (all levels.)", EUserInterfaceActionType::Button, FInputChord() );
-	UI_COMMAND( BuildLODsOnly, "Build LODs", "Only builds LODs (all levels.)", EUserInterfaceActionType::Button, FInputChord() );
+	UI_COMMAND( BuildHLODs, "Build HLODs", "Builds all HLODs for the current world", EUserInterfaceActionType::Button, FInputChord() );
+	UI_COMMAND( BuildMinimap, "Build Minimap", "Builds the minimap for the curent world", EUserInterfaceActionType::Button, FInputChord() );
 	UI_COMMAND( BuildTextureStreamingOnly, "Build Texture Streaming", "Build texture streaming data", EUserInterfaceActionType::Button, FInputChord() );
 	UI_COMMAND( BuildVirtualTextureOnly, "Build Virtual Textures", "Build runtime virtual texture low mips streaming data", EUserInterfaceActionType::Button, FInputChord());
 	UI_COMMAND( BuildAllLandscape, "Build All Landscape", "Build All Landscape Data(Grass, Baked Textures)", EUserInterfaceActionType::Button, FInputChord());
diff --git a/Engine/Source/Editor/LevelEditor/Private/LevelEditorMenu.cpp b/Engine/Source/Editor/LevelEditor/Private/LevelEditorMenu.cpp
--- a/Engine/Source/Editor/LevelEditor/Private/LevelEditorMenu.cpp
+++ b/Engine/Source/Editor/LevelEditor/Private/LevelEditorMenu.cpp
@@ -636,8 +636,9 @@ void FLevelEditorMenu::RegisterBuildMenu()
 	}
 
 	{
-		FToolMenuSection& Section = Menu->AddSection("LevelEditorLOD", LOCTEXT("LODHeading", "Hierarchical LOD"));
-		Section.AddMenuEntry(FLevelEditorCommands::Get().BuildLODsOnly);
+		FToolMenuSection& Section = Menu->AddSection("LevelEditorWorldPartition", LOCTEXT("WorldPartitionHeading", "World Partition"));
+		Section.AddMenuEntry(FLevelEditorCommands::Get().BuildHLODs);
+		Section.AddMenuEntry(FLevelEditorCommands::Get().BuildMinimap);
 	}
 
 	{
diff --git a/Engine/Source/Editor/LevelEditor/Private/SLevelEditor.cpp b/Engine/Source/Editor/LevelEditor/Private/SLevelEditor.cpp
--- a/Engine/Source/Editor/LevelEditor/Private/SLevelEditor.cpp
+++ b/Engine/Source/Editor/LevelEditor/Private/SLevelEditor.cpp
@@ -950,6 +950,13 @@ bool SLevelEditor::CanSpawnLevelEditorTab(const FSpawnTabArgs& Args, FName TabId
 	{
 		return false;
 	}
+
+	// HLOD Outliner not yet supported with World Partition
+	if (TabIdentifier == LevelEditorTabIds::LevelEditorHierarchicalLODOutliner && World->IsPartitionedWorld())
+	{
+		return false;
+	}
+
 	return true;
 }
 
@@ -1263,7 +1270,9 @@ TSharedRef<SWidget> SLevelEditor::RestoreContentArea( const TSharedRef<SDockTab>
 
 		{
 			const FSlateIcon LayersIcon(FEditorStyle::GetStyleSetName(), "LevelEditor.Tabs.HLOD");
-			LevelEditorTabManager->RegisterTabSpawner(LevelEditorTabIds::LevelEditorHierarchicalLODOutliner, FOnSpawnTab::CreateSP<SLevelEditor, FName, FString>(this, &SLevelEditor::SpawnLevelEditorTab, LevelEditorTabIds::LevelEditorHierarchicalLODOutliner, FString()))
+			LevelEditorTabManager->RegisterTabSpawner(LevelEditorTabIds::LevelEditorHierarchicalLODOutliner, 
+				FOnSpawnTab::CreateSP<SLevelEditor, FName, FString>(this, &SLevelEditor::SpawnLevelEditorTab, LevelEditorTabIds::LevelEditorHierarchicalLODOutliner, FString()),
+				FCanSpawnTab::CreateSP<SLevelEditor, FName>(this, &SLevelEditor::CanSpawnLevelEditorTab, LevelEditorTabIds::LevelEditorHierarchicalLODOutliner))
 				.SetDisplayName(NSLOCTEXT("LevelEditorTabs", "LevelEditorHierarchicalLODOutliner", "Hierarchical LOD Outliner"))
 				.SetTooltipText(NSLOCTEXT("LevelEditorTabs", "LevelEditorHierarchicalLODOutlinerTooltipText", "Open the Hierarchical LOD Outliner."))
 				.SetGroup(MenuStructure.GetLevelEditorCategory())
diff --git a/Engine/Source/Editor/LevelEditor/Public/LevelEditorActions.h b/Engine/Source/Editor/LevelEditor/Public/LevelEditorActions.h
--- a/Engine/Source/Editor/LevelEditor/Public/LevelEditorActions.h
+++ b/Engine/Source/Editor/LevelEditor/Public/LevelEditorActions.h
@@ -77,7 +77,8 @@ public:
 	TSharedPtr< FUICommandInfo > BuildGeometryOnly;
 	TSharedPtr< FUICommandInfo > BuildGeometryOnly_OnlyCurrentLevel;
 	TSharedPtr< FUICommandInfo > BuildPathsOnly;
-	TSharedPtr< FUICommandInfo > BuildLODsOnly;
+	TSharedPtr< FUICommandInfo > BuildHLODs;
+	TSharedPtr< FUICommandInfo > BuildMinimap;
 	TSharedPtr< FUICommandInfo > BuildTextureStreamingOnly;
 	TSharedPtr< FUICommandInfo > BuildVirtualTextureOnly;
 	TSharedPtr< FUICommandInfo > BuildAllLandscape;
@@ -748,7 +749,8 @@ public:
 	static void BuildGeometryOnly_Execute();
 	static void BuildGeometryOnly_OnlyCurrentLevel_Execute();
 	static void BuildPathsOnly_Execute();
-	static void BuildLODsOnly_Execute();
+	static void BuildHLODs_Execute();
+	static void BuildMinimap_Execute();
 	static void BuildTextureStreamingOnly_Execute();
 	static void BuildVirtualTextureOnly_Execute();
 	static void BuildAllLandscape_Execute();
diff --git a/Engine/Source/Editor/UnrealEd/Private/EditorBuildUtils.cpp b/Engine/Source/Editor/UnrealEd/Private/EditorBuildUtils.cpp
--- a/Engine/Source/Editor/UnrealEd/Private/EditorBuildUtils.cpp
+++ b/Engine/Source/Editor/UnrealEd/Private/EditorBuildUtils.cpp
@@ -49,6 +49,7 @@
 #include "ShaderCompilerCore.h"
 #include "AssetRegistryModule.h"
 #include "Interfaces/IMainFrameModule.h"
+#include "WorldPartition/IWorldPartitionEditorModule.h"
 #include "WorldPartition/SWorldPartitionBuildNavigationDialog.h"
 #include "WorldPartition/WorldPartitionBuildNavigationOptions.h"
 
@@ -69,6 +70,7 @@ const FName FBuildOptions::BuildAll(TEXT("BuildAll"));
 const FName FBuildOptions::BuildAllSubmit(TEXT("BuildAllSubmit"));
 const FName FBuildOptions::BuildAllOnlySelectedPaths(TEXT("BuildAllOnlySelectedPaths"));
 const FName FBuildOptions::BuildHierarchicalLOD(TEXT("BuildHierarchicalLOD"));
+const FName FBuildOptions::BuildMinimap(TEXT("BuildMinimap"));
 const FName FBuildOptions::BuildTextureStreaming(TEXT("BuildTextureStreaming"));
 const FName FBuildOptions::BuildVirtualTexture(TEXT("BuildVirtualTexture"));
 const FName FBuildOptions::BuildAllLandscape(TEXT("BuildAllLandscape"));
@@ -306,7 +308,11 @@ bool FEditorBuildUtils::EditorBuild( UWorld* InWorld, FName Id, const bool bAllo
 	}
 	else if (Id == FBuildOptions::BuildHierarchicalLOD)
 	{
-		BuildType = SBuildProgressWidget::BUILDTYPE_LODs;
+		BuildType = SBuildProgressWidget::BUILDTYPE_HLODs;
+	}
+	else if (Id == FBuildOptions::BuildMinimap)
+	{
+		BuildType = SBuildProgressWidget::BUILDTYPE_Minimap;
 	}
 	else if (Id == FBuildOptions::BuildTextureStreaming)
 	{
@@ -432,15 +438,38 @@ bool FEditorBuildUtils::EditorBuild( UWorld* InWorld, FName Id, const bool bAllo
 		bDoBuild = GEditor->WarnAboutHiddenLevels( InWorld, false );
 		if ( bDoBuild )
 		{
-				GEditor->ResetTransaction( NSLOCTEXT("UnrealEd", "BuildHLODMeshes", "Building Hierarchical LOD Meshes") );
+			GEditor->ResetTransaction( NSLOCTEXT("UnrealEd", "BuildHLODMeshes", "Building Hierarchical LOD Meshes") );
 
 			// We can't set the busy cursor for all windows, because lighting
 			// needs a cursor for the lighting options dialog.
 			const FScopedBusyCursor BusyCursor;
 
+			if (InWorld->IsPartitionedWorld())
+			{
+				bShouldMapCheck = false;
+				bDirtyPersistentLevel = false;
+			}
+
 			TriggerHierarchicalLODBuilder(InWorld, Id);
 		}
 	}
+	else if (Id == FBuildOptions::BuildMinimap)
+	{
+		bDoBuild = InWorld->IsPartitionedWorld();
+		if ( bDoBuild )
+		{
+			GEditor->ResetTransaction( NSLOCTEXT("UnrealEd", "BuildHLODMeshes", "Building Hierarchical LOD Meshes") );
+
+			// We can't set the busy cursor for all windows, because lighting
+			// needs a cursor for the lighting options dialog.
+			const FScopedBusyCursor BusyCursor;
+
+			bShouldMapCheck = false;
+			bDirtyPersistentLevel = false;
+
+			TriggerMinimapBuilder(InWorld, Id);
+		}
+	}
 	else if (Id == FBuildOptions::BuildAllLandscape)
 	{
 		bDoBuild = GEditor->WarnAboutHiddenLevels(InWorld, false);
@@ -1146,7 +1175,12 @@ void FBuildAllHandler::ProcessBuild(const TWeakPtr<SBuildProgressWidget>& BuildP
 		}
 		else if (StepId == FBuildOptions::BuildHierarchicalLOD)
 		{
-			BuildProgressWidget.Pin()->SetBuildType(SBuildProgressWidget::BUILDTYPE_LODs);
+			BuildProgressWidget.Pin()->SetBuildType(SBuildProgressWidget::BUILDTYPE_HLODs);
+			FEditorBuildUtils::TriggerHierarchicalLODBuilder(CurrentWorld, CurrentBuildId);
+		}
+		else if (StepId == FBuildOptions::BuildMinimap)
+		{
+			BuildProgressWidget.Pin()->SetBuildType(SBuildProgressWidget::BUILDTYPE_HLODs);
 			FEditorBuildUtils::TriggerHierarchicalLODBuilder(CurrentWorld, CurrentBuildId);
 		}
 		else if (StepId == FBuildOptions::BuildTextureStreaming)
@@ -1231,8 +1265,27 @@ void FBuildAllHandler::BuildFinished()
 
 void FEditorBuildUtils::TriggerHierarchicalLODBuilder(UWorld* InWorld, FName Id)
 {
-	// Invoke HLOD generator, with either preview or full build
-	InWorld->HierarchicalLODBuilder->BuildMeshesForLODActors(false);
+	if (InWorld->IsPartitionedWorld())
+	{
+		IWorldPartitionEditorModule& WorldPartitionEditorModule = FModuleManager::LoadModuleChecked<IWorldPartitionEditorModule>("WorldPartitionEditor");
+		TSubclassOf<UWorldPartitionBuilder> WorldPartitionHLODsBuilder = FindObjectChecked<UClass>(ANY_PACKAGE, TEXT("WorldPartitionHLODsBuilder"), true);
+		WorldPartitionEditorModule.RunBuilder(WorldPartitionHLODsBuilder, InWorld->GetPackage()->GetName());
+	}
+	else
+	{
+		// Invoke HLOD generator, with either preview or full build
+		InWorld->HierarchicalLODBuilder->BuildMeshesForLODActors(false);
+	}
+}
+
+void FEditorBuildUtils::TriggerMinimapBuilder(UWorld* InWorld, FName Id)
+{
+	if (InWorld->IsPartitionedWorld())
+	{
+		IWorldPartitionEditorModule& WorldPartitionEditorModule = FModuleManager::LoadModuleChecked<IWorldPartitionEditorModule>("WorldPartitionEditor");
+		TSubclassOf<UWorldPartitionBuilder> WorldPartitionMiniMapBuilder = FindObjectChecked<UClass>(ANY_PACKAGE, TEXT("WorldPartitionMiniMapBuilder"), true);
+		WorldPartitionEditorModule.RunBuilder(WorldPartitionMiniMapBuilder, InWorld->GetPackage()->GetName());
+	}
 }
 
 EDebugViewShaderMode ViewModeIndexToDebugViewShaderMode(EViewModeIndex SelectedViewMode)
diff --git a/Engine/Source/Editor/UnrealEd/Private/WorldPartition/WorldPartitionBuilder.cpp b/Engine/Source/Editor/UnrealEd/Private/WorldPartition/WorldPartitionBuilder.cpp
--- a/Engine/Source/Editor/UnrealEd/Private/WorldPartition/WorldPartitionBuilder.cpp
+++ b/Engine/Source/Editor/UnrealEd/Private/WorldPartition/WorldPartitionBuilder.cpp
@@ -205,11 +205,14 @@ bool UWorldPartitionBuilder::Run(UWorld* World, FPackageSourceControlHelper& Pac
 		const FIntVector MinCellCoords = GetCellCoord(CellInfo.EditorBounds.Min, IterativeCellSize);
 		const FIntVector MaxCellCoords = GetCellCoord(CellInfo.EditorBounds.Max, IterativeCellSize);
 
+		const int32 IterationCount = ((LoadingMode == ELoadingMode::IterativeCells2D) ? 1 : (MaxCellCoords.Z - MinCellCoords.Z + 1)) * (MaxCellCoords.Y - MinCellCoords.Y + 1) * (MaxCellCoords.X - MinCellCoords.X + 1);
+		int32 IterationIndex = 0;
+
 		UE_LOG(LogWorldPartitionBuilder, Display, TEXT("Iterative Cell Mode"));
 		UE_LOG(LogWorldPartitionBuilder, Display, TEXT("Cell Size %d"), IterativeCellSize);
 		UE_LOG(LogWorldPartitionBuilder, Display, TEXT("Cell Overlap %d"), IterativeCellOverlapSize);
 		UE_LOG(LogWorldPartitionBuilder, Display, TEXT("WorldBounds: Min %s, Max %s"), *CellInfo.EditorBounds.Min.ToString(), *CellInfo.EditorBounds.Max.ToString());
-		UE_LOG(LogWorldPartitionBuilder, Display, TEXT("Iteration Count: %d"), ((LoadingMode == ELoadingMode::IterativeCells2D) ? 1 : (MaxCellCoords.Z-MinCellCoords.Z + 1)) * (MaxCellCoords.Y-MinCellCoords.Y + 1) * (MaxCellCoords.X - MinCellCoords.X + 1));
+		UE_LOG(LogWorldPartitionBuilder, Display, TEXT("Iteration Count: %d"), IterationCount);
 		
 		FBox LoadedBounds(ForceInit);
 
@@ -219,6 +222,9 @@ bool UWorldPartitionBuilder::Run(UWorld* World, FPackageSourceControlHelper& Pac
 			{
 				for (int32 x = MinCellCoords.X; bResult && (x <= MaxCellCoords.X); x++)
 				{
+					IterationIndex++;
+					UE_LOG(LogWorldPartitionBuilder, Display, TEXT("[%d / %d] Processing cells..."), IterationIndex, IterationCount);
+
 					FVector Min(x * IterativeCellSize, y * IterativeCellSize, z * IterativeCellSize);
 					FVector Max = Min + FVector(IterativeCellSize);
 
diff --git a/Engine/Source/Editor/UnrealEd/Private/WorldPartition/WorldPartitionHLODsBuilder.cpp b/Engine/Source/Editor/UnrealEd/Private/WorldPartition/WorldPartitionHLODsBuilder.cpp
--- a/Engine/Source/Editor/UnrealEd/Private/WorldPartition/WorldPartitionHLODsBuilder.cpp
+++ b/Engine/Source/Editor/UnrealEd/Private/WorldPartition/WorldPartitionHLODsBuilder.cpp
@@ -466,7 +466,7 @@ bool UWorldPartitionHLODsBuilder::BuildHLODActors()
 
 		AWorldPartitionHLOD* HLODActor = CastChecked<AWorldPartitionHLOD>(ActorDesc->GetActor());
 
-		UE_LOG(LogWorldPartitionHLODsBuilder, Display, TEXT("    [%d/%d] Building HLOD actor %s..."), CurrentActor, HLODActorsToBuild.Num(), *HLODActor->GetActorLabel());
+		UE_LOG(LogWorldPartitionHLODsBuilder, Display, TEXT("[%d / %d] Building HLOD actor %s..."), CurrentActor + 1, HLODActorsToBuild.Num(), *HLODActor->GetActorLabel());
 
 		if (HLODLevelToBuild == INDEX_NONE || HLODActor->GetLODLevel() == HLODLevelToBuild)
 		{
@@ -548,23 +548,23 @@ bool UWorldPartitionHLODsBuilder::DeleteHLODActors()
 	WorldPartition->Uninitialize();
 	FWorldPartitionHelpers::DoCollectGarbage();
 
-	int32 NumDeleted = 0;
-	for (const FString& PackageName : PackagesToDelete)
+	for (int32 PackageIndex = 0; PackageIndex < PackagesToDelete.Num(); ++PackageIndex)
 	{
+		const FString& PackageName = PackagesToDelete[PackageIndex];
+
 		bool bDeleted = SourceControlHelper->Delete(PackageName);
 		if (bDeleted)
 		{
-			UE_LOG(LogWorldPartitionHLODsBuilder, Display, TEXT("    [%d] %s"), NumDeleted, *PackageName);
+			UE_LOG(LogWorldPartitionHLODsBuilder, Display, TEXT("[%d / %d] Deleting %s..."), PackageIndex + 1, PackagesToDelete.Num(), *PackageName);
 		}
 		else
 		{
 			UE_LOG(LogWorldPartitionHLODsBuilder, Error, TEXT("Failed to delete %s, exiting..."), *PackageName);
 			return false;
 		}
-		NumDeleted++;
 	}
 
-	UE_LOG(LogWorldPartitionHLODsBuilder, Display, TEXT("#### Deleted %d HLOD actors ####"), NumDeleted);
+	UE_LOG(LogWorldPartitionHLODsBuilder, Display, TEXT("#### Deleted %d HLOD actors ####"), PackagesToDelete.Num());
 
 	return true;
 }
diff --git a/Engine/Source/Editor/UnrealEd/Public/Dialogs/SBuildProgress.h b/Engine/Source/Editor/UnrealEd/Public/Dialogs/SBuildProgress.h
--- a/Engine/Source/Editor/UnrealEd/Public/Dialogs/SBuildProgress.h
+++ b/Engine/Source/Editor/UnrealEd/Public/Dialogs/SBuildProgress.h
@@ -40,8 +40,10 @@ public:
 		BUILDTYPE_Lighting,
 		/** Paths are being built. */
 		BUILDTYPE_Paths,
-		/** LODs are being built */
-		BUILDTYPE_LODs, 
+		/** HLODs are being built */
+		BUILDTYPE_HLODs,
+		/** Minimap is being built */
+		BUILDTYPE_Minimap,
 		/** Texture streaming data is being built */
 		BUILDTYPE_TextureStreaming, 
 		/** Virtual texture data is being built */
diff --git a/Engine/Source/Editor/UnrealEd/Public/EditorBuildUtils.h b/Engine/Source/Editor/UnrealEd/Public/EditorBuildUtils.h
--- a/Engine/Source/Editor/UnrealEd/Public/EditorBuildUtils.h
+++ b/Engine/Source/Editor/UnrealEd/Public/EditorBuildUtils.h
@@ -30,8 +30,10 @@ struct FBuildOptions
 	UNREALED_API static const FName BuildAllSubmit;
 	/** Build everything except for paths only build selected */
 	UNREALED_API static const FName BuildAllOnlySelectedPaths;
-	/** Build Hierarchical LOD system - need WorldSetting setup*/
+	/** Build HLODs */
 	UNREALED_API static const FName BuildHierarchicalLOD;
+	/** Build minimap */
+	UNREALED_API static const FName BuildMinimap;
 	/** Build texture streaming */
 	UNREALED_API static const FName BuildTextureStreaming;
 	/** Build virtual textures */
@@ -262,13 +264,21 @@ private:
 	static bool WorldPartitionBuildNavigation(const FString& InLongPackageName);
 
 	/** 
-	 * Trigger LOD builder to (re)generate LODActors
+	 * Trigger HLOD builder to (re)generate HLOD actors
 	 *
 	 * @param	InWorld			WorldContext
 	 * @param	BuildSettings	Build settings that will be used for the editor build
 	 */
 	static void TriggerHierarchicalLODBuilder(UWorld* InWorld, FName Id);
 
+	/** 
+	 * Trigger minimap builder to (re)generate minimap
+	 *
+	 * @param	InWorld			WorldContext
+	 * @param	BuildSettings	Build settings that will be used for the editor build
+	 */
+	static void TriggerMinimapBuilder(UWorld* InWorld, FName Id);
+
 	/** Intentionally hide constructors, etc. to prevent instantiation */
 	FEditorBuildUtils();
 	~FEditorBuildUtils();
diff --git a/Engine/Source/Editor/UnrealEd/Public/WorldPartition/WorldPartitionHLODsBuilder.h b/Engine/Source/Editor/UnrealEd/Public/WorldPartition/WorldPartitionHLODsBuilder.h
--- a/Engine/Source/Editor/UnrealEd/Public/WorldPartition/WorldPartitionHLODsBuilder.h
+++ b/Engine/Source/Editor/UnrealEd/Public/WorldPartition/WorldPartitionHLODsBuilder.h
@@ -70,7 +70,7 @@ enum class EHLODBuildStep : uint8
 ENUM_CLASS_FLAGS(EHLODBuildStep);
 
 UCLASS()
-class UWorldPartitionHLODsBuilder : public UWorldPartitionBuilder
+class UNREALED_API UWorldPartitionHLODsBuilder : public UWorldPartitionBuilder
 {
 	GENERATED_UCLASS_BODY()
 public:
diff --git a/Engine/Source/Editor/UnrealEd/Public/WorldPartition/WorldPartitionMiniMapBuilder.h b/Engine/Source/Editor/UnrealEd/Public/WorldPartition/WorldPartitionMiniMapBuilder.h
--- a/Engine/Source/Editor/UnrealEd/Public/WorldPartition/WorldPartitionMiniMapBuilder.h
+++ b/Engine/Source/Editor/UnrealEd/Public/WorldPartition/WorldPartitionMiniMapBuilder.h
@@ -8,7 +8,7 @@
 #include "WorldPartitionMiniMapBuilder.generated.h"
 
 UCLASS(config = Engine, defaultconfig)
-class UWorldPartitionMiniMapBuilder : public UWorldPartitionBuilder
+class UNREALED_API UWorldPartitionMiniMapBuilder : public UWorldPartitionBuilder
 {
 	GENERATED_UCLASS_BODY()
 
diff --git a/Engine/Source/Editor/WorldPartitionEditor/Private/WorldPartition/HLOD/SWorldPartitionBuildHLODsDialog.cpp b/Engine/Source/Editor/WorldPartitionEditor/Private/WorldPartition/HLOD/SWorldPartitionBuildHLODsDialog.cpp
new file mode 100644
--- /dev/null
+++ b/Engine/Source/Editor/WorldPartitionEditor/Private/WorldPartition/HLOD/SWorldPartitionBuildHLODsDialog.cpp
@@ -0,0 +1,103 @@
+// Copyright Epic Games, Inc. All Rights Reserved.
+#include "SWorldPartitionBuildHLODsDialog.h"
+
+#include "Widgets/Layout/SSpacer.h"
+#include "Widgets/Input/SButton.h"
+#include "Widgets/Text/STextBlock.h"
+#include "EditorStyleSet.h"
+
+#define LOCTEXT_NAMESPACE "WorldPartitionBuildHLODsDialog"
+
+const FVector2D SWorldPartitionBuildHLODsDialog::DEFAULT_WINDOW_SIZE = FVector2D(600, 200);
+
+void SWorldPartitionBuildHLODsDialog::Construct(const FArguments& InArgs)
+{
+	ParentWindowPtr = InArgs._ParentWindow.Get();
+	Result = DialogResult::Cancel;
+
+	this->ChildSlot
+		[
+			SNew(SBorder)
+			.BorderImage(FEditorStyle::GetBrush("ToolPanel.GroupBorder"))
+			[
+				SNew(SVerticalBox)
+				+ SVerticalBox::Slot()
+				.MaxHeight(200.0f)
+				[
+					// Options here
+					SNew(SHorizontalBox)
+					+ SHorizontalBox::Slot()
+					.Padding(20, 20, 20, 20)
+					[
+						SNew(STextBlock)
+						.AutoWrapText(true)
+						.Text(LOCTEXT("BuildHLODDescription", "Select the operation you wish to perform. Note that this will require the current map be saved & unloaded. It will be reopened once the operation is completed."))
+					]
+				]
+				+ SVerticalBox::Slot()
+				.AutoHeight()
+				[
+					SNew(SHorizontalBox)
+					+ SHorizontalBox::Slot()
+					.FillWidth(1.f)
+					[
+						SNew(SSpacer)
+					]
+					+ SHorizontalBox::Slot()
+					.AutoWidth()
+					.Padding(2.f)
+					[
+						SNew(SButton)
+						.HAlign(HAlign_Center)
+						.ContentPadding(FEditorStyle::GetMargin("StandardDialog.ContentPadding"))
+						.OnClicked(this, &SWorldPartitionBuildHLODsDialog::OnBuildClicked)
+						.Text(LOCTEXT("BuildButton", "Build HLODs"))
+					]
+					+ SHorizontalBox::Slot()
+						.AutoWidth()
+						.Padding(2.f)
+						[
+							SNew(SButton)
+							.HAlign(HAlign_Center)
+						.ContentPadding(FEditorStyle::GetMargin("StandardDialog.ContentPadding"))
+						.OnClicked(this, &SWorldPartitionBuildHLODsDialog::OnDeleteClicked)
+						.Text(LOCTEXT("DeleteButton", "Delete HLODs"))
+						]
+					+ SHorizontalBox::Slot()
+					.AutoWidth()
+					.Padding(2.f)
+					[
+						SNew(SButton)
+						.HAlign(HAlign_Center)
+						.ContentPadding(FEditorStyle::GetMargin("StandardDialog.ContentPadding"))
+						.OnClicked(this, &SWorldPartitionBuildHLODsDialog::OnCancelClicked)
+						.Text(LOCTEXT("CancelButton", "Cancel"))
+					]
+				]
+			]
+		];
+}
+
+FReply SWorldPartitionBuildHLODsDialog::OnBuildClicked()
+{
+	Result = DialogResult::BuildHLODs;
+	ParentWindowPtr.Pin()->RequestDestroyWindow();
+	return FReply::Handled();
+}
+
+
+FReply SWorldPartitionBuildHLODsDialog::OnDeleteClicked()
+{
+	Result = DialogResult::DeleteHLODs;
+	ParentWindowPtr.Pin()->RequestDestroyWindow();
+	return FReply::Handled();
+}
+
+FReply SWorldPartitionBuildHLODsDialog::OnCancelClicked()
+{
+	Result = DialogResult::Cancel;
+	ParentWindowPtr.Pin()->RequestDestroyWindow();
+	return FReply::Handled();
+}
+
+#undef LOCTEXT_NAMESPACE
diff --git a/Engine/Source/Editor/WorldPartitionEditor/Private/WorldPartition/HLOD/SWorldPartitionBuildHLODsDialog.h b/Engine/Source/Editor/WorldPartitionEditor/Private/WorldPartition/HLOD/SWorldPartitionBuildHLODsDialog.h
new file mode 100644
--- /dev/null
+++ b/Engine/Source/Editor/WorldPartitionEditor/Private/WorldPartition/HLOD/SWorldPartitionBuildHLODsDialog.h
@@ -0,0 +1,56 @@
+// Copyright Epic Games, Inc. All Rights Reserved.
+#pragma once
+
+#include "CoreMinimal.h"
+#include "InputCoreTypes.h"
+#include "Types/SlateEnums.h"
+#include "Styling/SlateColor.h"
+#include "Input/Reply.h"
+#include "Widgets/SWidget.h"
+#include "Widgets/SWindow.h"
+#include "Widgets/DeclarativeSyntaxSupport.h"
+#include "Widgets/SCompoundWidget.h"
+
+//////////////////////////////////////////////////////////////////////////
+// SWorldPartitionBuildHLODsDialog
+
+class SWorldPartitionBuildHLODsDialog : public SCompoundWidget
+{
+public:
+	enum DialogResult
+	{
+		BuildHLODs,
+		DeleteHLODs,
+		Cancel
+	};
+
+	SLATE_BEGIN_ARGS(SWorldPartitionBuildHLODsDialog) {}
+		/** A pointer to the parent window */
+		SLATE_ATTRIBUTE(TSharedPtr<SWindow>, ParentWindow)
+	SLATE_END_ARGS()
+
+	~SWorldPartitionBuildHLODsDialog() {}
+
+	/**
+	 * Construct this widget.  Called by the SNew() Slate macro.
+	 *
+	 * @param	InArgs				Declaration used by the SNew() macro to construct this widget
+	 * @param	InViewModel			The UI logic not specific to slate
+	 */
+	void Construct(const FArguments& InArgs);
+
+	static const FVector2D DEFAULT_WINDOW_SIZE;
+
+	DialogResult GetDialogResult() const { return Result; }
+	
+private:
+	FReply OnBuildClicked();
+	FReply OnDeleteClicked();
+	FReply OnCancelClicked();
+		
+	/** Pointer to the parent window, so we know to destroy it when done */
+	TWeakPtr<SWindow> ParentWindowPtr;
+		
+	/** Dialog Result */
+	DialogResult Result;
+};
\ No newline at end of file
diff --git a/Engine/Source/Editor/WorldPartitionEditor/Private/WorldPartitionEditorModule.cpp b/Engine/Source/Editor/WorldPartitionEditor/Private/WorldPartitionEditorModule.cpp
--- a/Engine/Source/Editor/WorldPartitionEditor/Private/WorldPartitionEditorModule.cpp
+++ b/Engine/Source/Editor/WorldPartitionEditor/Private/WorldPartitionEditorModule.cpp
@@ -4,6 +4,9 @@
 #include "WorldPartition/WorldPartition.h"
 #include "WorldPartition/WorldPartitionVolume.h"
 
+#include "WorldPartition/WorldPartitionHLODsBuilder.h"
+#include "WorldPartition/WorldPartitionMiniMapBuilder.h"
+
 #include "WorldPartition/HLOD/HLODLayerAssetTypeActions.h"
 
 #include "WorldPartition/SWorldPartitionEditor.h"
@@ -15,6 +18,8 @@
 #include "WorldPartition/WorldPartitionConvertOptions.h"
 #include "WorldPartition/WorldPartitionEditorSettings.h"
 
+#include "WorldPartition/HLOD/SWorldPartitionBuildHLODsDialog.h"
+
 #include "LevelEditor.h"
 #include "Framework/MultiBox/MultiBoxBuilder.h"
 
@@ -23,6 +28,7 @@
 #include "Misc/MessageDialog.h"
 #include "Misc/StringBuilder.h"
 #include "Misc/ScopedSlowTask.h"
+#include "Internationalization/Regex.h"
 
 #include "Interfaces/IMainFrameModule.h"
 #include "Widgets/SWindow.h"
@@ -30,6 +36,8 @@
 #include "FileHelpers.h"
 #include "ToolMenus.h"
 #include "IContentBrowserSingleton.h"
+#include "IDirectoryWatcher.h"
+#include "DirectoryWatcherModule.h"
 #include "ContentBrowserModule.h"
 #include "EditorDirectories.h"
 #include "AssetRegistryModule.h"
@@ -45,6 +53,8 @@ IMPLEMENT_MODULE( FWorldPartitionEditorModule, WorldPartitionEditor );
 
 const FName WorldPartitionEditorTabId("WorldBrowserPartitionEditor");
 
+DEFINE_LOG_CATEGORY_STATIC(LogWorldPartitionEditor, All, All);
+
 // World Partition
 static void OnLoadSelectedWorldPartitionVolumes(TArray<TWeakObjectPtr<AActor>> Volumes)
 {
@@ -212,6 +222,103 @@ void FWorldPartitionEditorModule::OnConvertMap()
 	}
 }
 
+static bool UnloadCurrentMap(bool bAskSaveContentPackages)
+{
+	// Ask user to save dirty packages
+	if (!FEditorFileUtils::SaveDirtyPackages(/*bPromptUserToSave=*/true, /*bSaveMapPackages=*/true, bAskSaveContentPackages))
+	{
+		return false;
+	}
+
+	// Unload any loaded map
+	if (!UEditorLoadingAndSavingUtils::NewBlankMap(/*bSaveExistingMap*/false))
+	{
+		return false;
+	}
+
+	return true;
+}
+
+static void RescanAssetsAndLoadMap(const FString& MapToLoad)
+{
+	// Force a directory watcher tick for the asset registry to get notified of the changes
+	FDirectoryWatcherModule& DirectoryWatcherModule = FModuleManager::Get().LoadModuleChecked<FDirectoryWatcherModule>(TEXT("DirectoryWatcher"));
+	DirectoryWatcherModule.Get()->Tick(-1.0f);
+
+	// Force update before loading converted map
+	FAssetRegistryModule& AssetRegistryModule = FModuleManager::LoadModuleChecked<FAssetRegistryModule>("AssetRegistry");
+	IAssetRegistry& AssetRegistry = AssetRegistryModule.Get();
+
+	TArray<FString> ExternalObjectsPaths = ULevel::GetExternalObjectsPaths(MapToLoad);
+
+	AssetRegistry.ScanModifiedAssetFiles({ MapToLoad });
+	AssetRegistry.ScanPathsSynchronous(ExternalObjectsPaths, true);
+
+	FEditorFileUtils::LoadMap(MapToLoad);
+}
+
+static void RunCommandletAsExternalProcess(const FString& InCommandletArgs, const FText& InOperationDescription, int32& OutResult, bool& bOutCancelled, FString& OutCommandletOutput)
+{
+	OutResult = 0;
+	bOutCancelled = false;
+	OutCommandletOutput.Empty();
+
+	FProcHandle ProcessHandle;
+
+	FScopedSlowTask SlowTask(1.0f, InOperationDescription);
+	SlowTask.MakeDialog(true);
+
+	void* ReadPipe = nullptr;
+	void* WritePipe = nullptr;
+	verify(FPlatformProcess::CreatePipe(ReadPipe, WritePipe));
+
+	FString CurrentExecutableName = FPlatformProcess::ExecutablePath();
+
+	// Try to provide complete Path, if we can't try with project name
+	FString ProjectPath = FPaths::IsProjectFilePathSet() ? FPaths::GetProjectFilePath() : FApp::GetProjectName();
+
+	uint32 ProcessID;
+	FString Arguments = FString::Printf(TEXT("\"%s\" %s"), *ProjectPath, *InCommandletArgs);
+	
+	UE_LOG(LogWorldPartitionEditor, Display, TEXT("Running commandlet: %s %s"), *CurrentExecutableName, *Arguments);
+	ProcessHandle = FPlatformProcess::CreateProc(*CurrentExecutableName, *Arguments, true, false, false, &ProcessID, 0, nullptr, WritePipe, ReadPipe);
+
+	while (FPlatformProcess::IsProcRunning(ProcessHandle))
+	{
+		if (SlowTask.ShouldCancel())
+		{
+			bOutCancelled = true;
+			FPlatformProcess::TerminateProc(ProcessHandle);
+			break;
+		}
+
+		const FString LogString = FPlatformProcess::ReadPipe(ReadPipe);
+		if (!LogString.IsEmpty())
+		{
+			OutCommandletOutput.Append(LogString);
+		}
+
+		// Parse output, look for progress indicator in the log (in the form "Display: [i / N] Msg...\n")
+		const FRegexPattern LogProgressPattern(TEXT("Display:\\s\\[([0-9]+)\\s\\/\\s([0-9]+)\\]\\s(.+)?(?=\\.{3}$)"));
+		FRegexMatcher Regex(LogProgressPattern, *LogString);
+		while (Regex.FindNext())
+		{
+			// Update slow task progress & message
+			SlowTask.CompletedWork = FCString::Atoi(*Regex.GetCaptureGroup(1));
+			SlowTask.TotalAmountOfWork = FCString::Atoi(*Regex.GetCaptureGroup(2));
+			SlowTask.DefaultMessage = FText::FromString(Regex.GetCaptureGroup(3));
+		}
+
+		SlowTask.EnterProgressFrame(0);
+		FPlatformProcess::Sleep(0.1);
+	}
+
+	UE_LOG(LogWorldPartitionEditor, Display, TEXT("#### Begin commandlet output ####\n%s"), *OutCommandletOutput);
+	UE_LOG(LogWorldPartitionEditor, Display, TEXT("#### End commandlet output ####"));
+
+	FPlatformProcess::ClosePipe(ReadPipe, WritePipe);
+}
+
 bool FWorldPartitionEditorModule::ConvertMap(const FString& InLongPackageName)
 {
 	if (ULevel::GetIsLevelPartitionedFromPackage(FName(*InLongPackageName)))
@@ -247,108 +354,159 @@ bool FWorldPartitionEditorModule::ConvertMap(const FString& InLongPackageName)
 
 	if (ConvertDialog->ClickedOk())
 	{
-		// Conversion will try to load the converted map so ask user to save dirty packages
-		if (!FEditorFileUtils::SaveDirtyPackages(/*bPromptUserToSave=*/true, /*bSaveMapPackages=*/true, /*bSaveContentPackages=*/false))
+		if (!UnloadCurrentMap(/*bAskSaveContentPackages=*/false))
 		{
 			return false;
 		}
 
-		// Unload any loaded map
-		if (!UEditorLoadingAndSavingUtils::NewBlankMap(/*bSaveExistingMap*/false))
+		const FString CommandletArgs = *DefaultConvertOptions->ToCommandletArgs();
+		const FText OperationDescription = LOCTEXT("ConvertProgress", "Converting map to world partition...");
+		
+		int32 Result;
+		bool bCancelled;
+		FString CommandletOutput;
+		RunCommandletAsExternalProcess(CommandletArgs, OperationDescription, Result, bCancelled, CommandletOutput);
+		if (!bCancelled && Result == 0)
+		{	
+			FMessageDialog::Open(EAppMsgType::Ok, FText::Format(LOCTEXT("ConvertMapCompleted", "Conversion completed:\n{0}"), FText::FromString(CommandletOutput)));
+
+#if	PLATFORM_DESKTOP
+			if (DefaultConvertOptions->bGenerateIni)
+			{
+				const FString PackageFilename = FPackageName::LongPackageNameToFilename(DefaultConvertOptions->LongPackageName);
+				const FString PackageDirectory = FPaths::ConvertRelativePathToFull(FPaths::GetPath(PackageFilename));
+				FPlatformProcess::ExploreFolder(*PackageDirectory);
+			}
+#endif				
+				
+			FString MapToLoad = InLongPackageName;
+			if (!DefaultConvertOptions->bInPlace)
+			{
+				MapToLoad += UWorldPartitionConvertCommandlet::GetConversionSuffix(DefaultConvertOptions->bOnlyMergeSubLevels);
+			}
+
+			RescanAssetsAndLoadMap(MapToLoad);
+		}
+		else if (bCancelled)
 		{
-			return false;
+			FMessageDialog::Open(EAppMsgType::Ok, LOCTEXT("ConvertMapCancelled", "Conversion cancelled!"));
 		}
+		else if(Result != 0)
+		{
+			FMessageDialog::Open(EAppMsgType::Ok, FText::Format(LOCTEXT("ConvertMapFailed", "Conversion failed:\n{0}"), FText::FromString(CommandletOutput)));
+		}
+	}
 
-		FProcHandle ProcessHandle;
-		FString CommandletOutput;
-		bool bCancelled = false;
+	return false;
+}
 
-		// Task scope
-		{
-			FScopedSlowTask SlowTask(0, LOCTEXT("ConvertProgress", "Converting map to world partition..."));
-			SlowTask.MakeDialog(true);
+bool FWorldPartitionEditorModule::RunBuilder(TSubclassOf<UWorldPartitionBuilder> InWorldPartitionBuilder, const FString& InLongPackageName)
+{
+	// Ideally this should be improved to automatically register all builders & present their options in a consistent way...
 
-			void* ReadPipe = nullptr;
-			void* WritePipe = nullptr;
-			verify(FPlatformProcess::CreatePipe(ReadPipe, WritePipe));
+	if (InWorldPartitionBuilder == UWorldPartitionHLODsBuilder::StaticClass())
+	{
+		return BuildHLODs(InLongPackageName);
+	}
+	
+	if (InWorldPartitionBuilder == UWorldPartitionMiniMapBuilder::StaticClass())
+	{
+		return BuildMinimap(InLongPackageName);
+	}
 
-			FString CurrentExecutableName = FPlatformProcess::ExecutablePath();
+	return false;
+}
 
-			// Try to provide complete Path, if we can't try with project name
-			FString ProjectPath = FPaths::IsProjectFilePathSet() ? FPaths::GetProjectFilePath() : FApp::GetProjectName();
+bool FWorldPartitionEditorModule::BuildHLODs(const FString& InMapToProcess)
+{
+	TSharedPtr<SWindow> DlgWindow =
+		SNew(SWindow)
+		.Title(LOCTEXT("BuildHLODsWindowTitle", "Build HLODs"))
+		.ClientSize(SWorldPartitionBuildHLODsDialog::DEFAULT_WINDOW_SIZE)
+		.SizingRule(ESizingRule::UserSized)
+		.SupportsMinimize(false)
+		.SupportsMaximize(false)
+		.SizingRule(ESizingRule::FixedSize);
 
-			uint32 ProcessID;
-			FString Arguments = FString::Printf(TEXT("\"%s\" %s"), *ProjectPath, *DefaultConvertOptions->ToCommandletArgs());
-			ProcessHandle = FPlatformProcess::CreateProc(*CurrentExecutableName, *Arguments, true, false, false, &ProcessID, 0, nullptr, WritePipe, ReadPipe);
-			
-			while (FPlatformProcess::IsProcRunning(ProcessHandle))
-			{
-				if (SlowTask.ShouldCancel())
-				{
-					bCancelled = true;
-					FPlatformProcess::TerminateProc(ProcessHandle);
-					break;
-				}
-
-				const FString LogString = FPlatformProcess::ReadPipe(ReadPipe);
-				if (!LogString.IsEmpty())
-				{
-					CommandletOutput.Append(LogString);
-				}
-
-				SlowTask.EnterProgressFrame(0);
-				FPlatformProcess::Sleep(0.1);
-			}
+	TSharedRef<SWorldPartitionBuildHLODsDialog> BuildHLODsDialog =
+		SNew(SWorldPartitionBuildHLODsDialog)
+		.ParentWindow(DlgWindow);
+
+	DlgWindow->SetContent(BuildHLODsDialog);
 
-			FPlatformProcess::ClosePipe(ReadPipe, WritePipe);
+	IMainFrameModule& MainFrameModule = FModuleManager::LoadModuleChecked<IMainFrameModule>(TEXT("MainFrame"));
+	FSlateApplication::Get().AddModalWindow(DlgWindow.ToSharedRef(), MainFrameModule.GetParentWindow());
+
+	if (BuildHLODsDialog->GetDialogResult() != SWorldPartitionBuildHLODsDialog::DialogResult::Cancel)
+	{
+		if (!UnloadCurrentMap(/*bAskSaveContentPackages=*/true))
+		{
+			return false;
 		}
 
-		int32 Result = 0;
-		if (!bCancelled && FPlatformProcess::GetProcReturnCode(ProcessHandle, &Result))
-		{	
-			if (Result == 0)
-			{
-				FMessageDialog::Open(EAppMsgType::Ok, FText::Format(LOCTEXT("ConvertMapCompleted", "Conversion completed:\n{0}"), FText::FromString(CommandletOutput)));
+		const FString BuildArgs = BuildHLODsDialog->GetDialogResult() == SWorldPartitionBuildHLODsDialog::DialogResult::BuildHLODs ? "-SetupHLODs -BuildHLODs -AllowCommandletRendering" : "-DeleteHLODs";
+		const FString CommandletArgs = InMapToProcess + " -run=WorldPartitionBuilderCommandlet -Builder=WorldPartitionHLODsBuilder " + BuildArgs;
+		const FText OperationDescription = LOCTEXT("HLODBuildProgress", "Building HLODs...");
 
-#if	PLATFORM_DESKTOP
-				if (DefaultConvertOptions->bGenerateIni)
-				{
-					const FString PackageFilename = FPackageName::LongPackageNameToFilename(DefaultConvertOptions->LongPackageName);
-					const FString PackageDirectory = FPaths::ConvertRelativePathToFull(FPaths::GetPath(PackageFilename));
-					FPlatformProcess::ExploreFolder(*PackageDirectory);
-				}
-#endif
-				
-				// Force update before loading converted map
-				FAssetRegistryModule& AssetRegistryModule = FModuleManager::LoadModuleChecked<FAssetRegistryModule>("AssetRegistry");
-				IAssetRegistry& AssetRegistry = AssetRegistryModule.Get();
-												
-				FString MapToLoad = InLongPackageName;
-				if (!DefaultConvertOptions->bInPlace)
-				{
-					MapToLoad += UWorldPartitionConvertCommandlet::GetConversionSuffix(DefaultConvertOptions->bOnlyMergeSubLevels);
-				}
-				
-				AssetRegistry.ScanModifiedAssetFiles({ MapToLoad });
-				AssetRegistry.ScanPathsSynchronous( ULevel::GetExternalObjectsPaths(MapToLoad), true);
-				
-				FEditorFileUtils::LoadMap(MapToLoad);
-			}
+		int32 Result;
+		bool bCancelled;
+		FString CommandletOutput;
+		RunCommandletAsExternalProcess(CommandletArgs, OperationDescription, Result, bCancelled, CommandletOutput);
+
+		bool bSuccess = !bCancelled && Result == 0;
+		if (bSuccess)
+		{
+			FMessageDialog::Open(EAppMsgType::Ok, FText::Format(LOCTEXT("HLODBuildCompleted", "HLOD build completed:\n{0}"), FText::FromString(CommandletOutput)));
+			RescanAssetsAndLoadMap(InMapToProcess);
 		}
 		else if (bCancelled)
 		{
-			FMessageDialog::Open(EAppMsgType::Ok, LOCTEXT("ConvertMapCancelled", "Conversion cancelled!"));
+			FMessageDialog::Open(EAppMsgType::Ok, LOCTEXT("HLODBuildCancelled", "HLOD build cancelled!"));
 		}
-		
-		if(Result != 0)
+		else if (Result != 0)
 		{
-			FMessageDialog::Open(EAppMsgType::Ok, LOCTEXT("ConvertMapFailed", "Conversion failed!"));
+			FMessageDialog::Open(EAppMsgType::Ok, FText::Format(LOCTEXT("HLODBuildFailed", "HLOD build failed:\n{0}"), FText::FromString(CommandletOutput)));
 		}
+
+		return bSuccess;
 	}
 
 	return false;
 }
 
+bool FWorldPartitionEditorModule::BuildMinimap(const FString& InMapToProcess)
+{
+	if (!UnloadCurrentMap(/*bAskSaveContentPackages=*/true))
+	{
+		return false;
+	}
+
+	const FString CommandletArgs = InMapToProcess + " -run=WorldPartitionBuilderCommandlet -Builder=WorldPartitionMinimapBuilder -AllowCommandletRendering";
+	const FText OperationDescription = LOCTEXT("MinimapBuildProgress", "Building minimap...");
+
+	int32 Result;
+	bool bCancelled;
+	FString CommandletOutput;
+	RunCommandletAsExternalProcess(CommandletArgs, OperationDescription, Result, bCancelled, CommandletOutput);
+	
+	bool bSuccess = !bCancelled && Result == 0;
+	if (bSuccess)
+	{
+		FMessageDialog::Open(EAppMsgType::Ok, FText::Format(LOCTEXT("MinimapBuildCompleted", "Minimap build completed:\n{0}"), FText::FromString(CommandletOutput)));
+		RescanAssetsAndLoadMap(InMapToProcess);
+	}
+	else if (bCancelled)
+	{
+		FMessageDialog::Open(EAppMsgType::Ok, LOCTEXT("MinimapBuildCancelled", "Minimap build cancelled!"));
+	}
+	else if (Result != 0)
+	{
+		FMessageDialog::Open(EAppMsgType::Ok, FText::Format(LOCTEXT("MinimapBuildFailed", "Minimap build failed:\n{0}"), FText::FromString(CommandletOutput)));
+	}
+
+	return bSuccess;
+}
+
 void FWorldPartitionEditorModule::OnMapChanged(uint32 MapFlags)
 {
 	if (MapFlags == MapChangeEventFlags::NewMap)
diff --git a/Engine/Source/Editor/WorldPartitionEditor/Public/WorldPartitionEditorModule.h b/Engine/Source/Editor/WorldPartitionEditor/Public/WorldPartitionEditorModule.h
--- a/Engine/Source/Editor/WorldPartitionEditor/Public/WorldPartitionEditorModule.h
+++ b/Engine/Source/Editor/WorldPartitionEditor/Public/WorldPartitionEditorModule.h
@@ -42,10 +42,15 @@ public:
 	virtual int32 GetInstancedFoliageGridSize() const;
 
 	/**
-	 * 
+	 * Convert the specified map to a world partition map.
 	 */
 	virtual bool ConvertMap(const FString& InLongPackageName) override;
 
+	/**
+	 * Run a world partition builder for the current map. Will display a dialog to specify options for the generation.
+	 */
+	virtual bool RunBuilder(TSubclassOf<UWorldPartitionBuilder> WorldPartitionBuilder, const FString& InLongPackageName) override;
+
 	/**
 	 *
 	 */
@@ -64,6 +69,9 @@ private:
 	/** Spawns the world partition tab */
 	TSharedRef<SDockTab> SpawnWorldPartitionTab(const FSpawnTabArgs& Args);
 
+	bool BuildHLODs(const FString& InMapToProcess);
+	bool BuildMinimap(const FString& InMapToProcess);
+
 private:
 	void OnConvertMap();
 
diff --git a/Engine/Source/Runtime/Engine/Private/WorldPartition/RuntimeSpatialHash/RuntimeSpatialHashHLOD.cpp b/Engine/Source/Runtime/Engine/Private/WorldPartition/RuntimeSpatialHash/RuntimeSpatialHashHLOD.cpp
--- a/Engine/Source/Runtime/Engine/Private/WorldPartition/RuntimeSpatialHash/RuntimeSpatialHashHLOD.cpp
+++ b/Engine/Source/Runtime/Engine/Private/WorldPartition/RuntimeSpatialHash/RuntimeSpatialHashHLOD.cpp
@@ -176,7 +176,7 @@ static TArray<FGuid> GenerateHLODsForGrid(UWorldPartition* WorldPartition, const
 
 				TRACE_CPUPROFILER_EVENT_SCOPE_TEXT(*CellName.ToString());
 
-				UE_LOG(LogWorldPartitionRuntimeSpatialHashHLOD, Display, TEXT("[%d / %d] Processing cell %s"), (int32)SlowTask.CompletedWork + 1, (int32)SlowTask.TotalAmountOfWork, *CellName.ToString());
+				UE_LOG(LogWorldPartitionRuntimeSpatialHashHLOD, Display, TEXT("[%d / %d] Processing cell %s..."), (int32)SlowTask.CompletedWork + 1, (int32)SlowTask.TotalAmountOfWork, *CellName.ToString());
 
 				FHLODCreationParams CreationParams;
 				CreationParams.WorldPartition = WorldPartition;
diff --git a/Engine/Source/Runtime/Engine/Public/WorldPartition/IWorldPartitionEditorModule.h b/Engine/Source/Runtime/Engine/Public/WorldPartition/IWorldPartitionEditorModule.h
--- a/Engine/Source/Runtime/Engine/Public/WorldPartition/IWorldPartitionEditorModule.h
+++ b/Engine/Source/Runtime/Engine/Public/WorldPartition/IWorldPartitionEditorModule.h
@@ -3,6 +3,9 @@
 
 #include "CoreMinimal.h"
 #include "Modules/ModuleInterface.h"
+#include "Templates/SubclassOf.h"
+
+class UWorldPartitionBuilder;
 
 /**
  * The module holding all of the UI related pieces for WorldPartition
@@ -14,6 +17,8 @@ public:
 
 	virtual bool ConvertMap(const FString& InLongPackageName) = 0;
 
+	virtual bool RunBuilder(TSubclassOf<UWorldPartitionBuilder> BuilderClass, const FString& InLongPackageName) = 0;
+
 	virtual int32 GetPlacementGridSize() const = 0;
 	virtual int32 GetInstancedFoliageGridSize() const = 0;
 
