fix non-unity build

#rnx

diff --git a/Engine/Source/Editor/WorldPartitionEditor/Private/WorldPartition/Customizations/WorldPartitionDetailsCustomization.h b/Engine/Source/Editor/WorldPartitionEditor/Private/WorldPartition/Customizations/WorldPartitionDetailsCustomization.h
--- a/Engine/Source/Editor/WorldPartitionEditor/Private/WorldPartition/Customizations/WorldPartitionDetailsCustomization.h
+++ b/Engine/Source/Editor/WorldPartitionEditor/Private/WorldPartition/Customizations/WorldPartitionDetailsCustomization.h
@@ -9,6 +9,8 @@
 class IDetailLayoutBuilder;
 class UWorldPartition;
 
+enum class ECheckBoxState : uint8;
+
 class FWorldPartitionDetails : public IDetailCustomization
 {
 public:
