World Partition Editor: fix displaying message when streaming is disabled.

#rnx

diff --git a/Engine/Source/Editor/WorldPartitionEditor/Private/WorldPartition/SWorldPartitionEditor.cpp b/Engine/Source/Editor/WorldPartitionEditor/Private/WorldPartition/SWorldPartitionEditor.cpp
--- a/Engine/Source/Editor/WorldPartitionEditor/Private/WorldPartition/SWorldPartitionEditor.cpp
+++ b/Engine/Source/Editor/WorldPartitionEditor/Private/WorldPartition/SWorldPartitionEditor.cpp
@@ -71,7 +71,12 @@ TSharedRef<SWidget> SWorldPartitionEditor::ConstructContentWidget()
 		if (UWorldPartition* WorldPartition = World->GetWorldPartition())
 		{
 			check(WorldPartition->World == World);
-			EditorName = WorldPartition->GetWorldPartitionEditorName();
+
+			if (WorldPartition->IsStreamingEnabled())
+			{
+				EditorName = WorldPartition->GetWorldPartitionEditorName();
+			}
+
 			WorldPartition->WorldPartitionEditor = this;
 		}
 	}
diff --git a/Engine/Source/Editor/WorldPartitionEditor/Public/WorldPartition/SWorldPartitionEditorGrid.h b/Engine/Source/Editor/WorldPartitionEditor/Public/WorldPartition/SWorldPartitionEditorGrid.h
--- a/Engine/Source/Editor/WorldPartitionEditor/Public/WorldPartition/SWorldPartitionEditorGrid.h
+++ b/Engine/Source/Editor/WorldPartitionEditor/Public/WorldPartition/SWorldPartitionEditorGrid.h
@@ -65,8 +65,8 @@ public:
 
 	void Refresh();
 
-	UWorld* World;
-	UWorldPartition* WorldPartition;
+	UWorld* World = nullptr;
+	UWorldPartition* WorldPartition = nullptr;
 
 protected:
 	static TMap<FName, PartitionEditorGridCreateInstanceFunc> PartitionEditorGridCreateInstanceFactory;
