World Partition: minor code cleanups.

#rnx

diff --git a/Engine/Source/Editor/WorldPartitionEditor/Private/WorldPartition/SWorldPartitionEditorGrid.cpp b/Engine/Source/Editor/WorldPartitionEditor/Private/WorldPartition/SWorldPartitionEditorGrid.cpp
--- a/Engine/Source/Editor/WorldPartitionEditor/Private/WorldPartition/SWorldPartitionEditorGrid.cpp
+++ b/Engine/Source/Editor/WorldPartitionEditor/Private/WorldPartition/SWorldPartitionEditorGrid.cpp
@@ -16,11 +16,10 @@ TMap<FName, SWorldPartitionEditorGrid::PartitionEditorGridCreateInstanceFunc> SW
 void SWorldPartitionEditorGrid::Construct(const FArguments& InArgs)
 {
 	World = InArgs._InWorld;
-	WorldPartition = World ? World->GetWorldPartition() : nullptr;
 
-	if (!WorldPartition || !WorldPartition->IsStreamingEnabled())
+	if (IsDisabled())
 	{
-		const FText Message = WorldPartition ? 
+		const FText Message = World->GetWorldPartition() ? 
 			LOCTEXT("WorldPartitionHasStreamingDisabled", "World Partition streaming is not enabled for this map") : 
 			LOCTEXT("WorldPartitionMustBeEnabled", "World Partition is not enabled for this map");
 
