-Fix crash (access on a NULL Wolrd ptr) that happens when switching to an non WP map when you have the World Partition editor opened and if you have opened the Levels editor while in the current map.

diff --git a/Engine/Source/Editor/WorldPartitionEditor/Private/WorldPartition/SWorldPartitionEditorGrid.cpp b/Engine/Source/Editor/WorldPartitionEditor/Private/WorldPartition/SWorldPartitionEditorGrid.cpp
--- a/Engine/Source/Editor/WorldPartitionEditor/Private/WorldPartition/SWorldPartitionEditorGrid.cpp
+++ b/Engine/Source/Editor/WorldPartitionEditor/Private/WorldPartition/SWorldPartitionEditorGrid.cpp
@@ -19,9 +19,10 @@ void SWorldPartitionEditorGrid::Construct(const FArguments& InArgs)
 
 	if (IsDisabled())
 	{
-		const FText Message = World->GetWorldPartition() ? 
-			LOCTEXT("WorldPartitionHasStreamingDisabled", "World Partition streaming is not enabled for this map") : 
-			LOCTEXT("WorldPartitionMustBeEnabled", "World Partition is not enabled for this map");
+		const FText Message = !World ? LOCTEXT("WorldPartitionNoWorldForEditorGrid", "No World") :
+									  World->GetWorldPartition() ? 
+									 	LOCTEXT("WorldPartitionHasStreamingDisabled", "World Partition streaming is not enabled for this map") : 
+										LOCTEXT("WorldPartitionMustBeEnabled", "World Partition is not enabled for this map");
 
 		ChildSlot
 		.VAlign(VAlign_Center)
