- Added a detail customization for WorldPartitionRuntimeSpatialHash
- Avoid dirtying level when changing the Preview Grids flag.

diff --git a/Engine/Source/Editor/WorldPartitionEditor/Private/WorldPartition/Customizations/WorldPartitionRuntimeSpatialHashDetailsCustomization.cpp b/Engine/Source/Editor/WorldPartitionEditor/Private/WorldPartition/Customizations/WorldPartitionRuntimeSpatialHashDetailsCustomization.cpp
new file mode 100644
--- /dev/null
+++ b/Engine/Source/Editor/WorldPartitionEditor/Private/WorldPartition/Customizations/WorldPartitionRuntimeSpatialHashDetailsCustomization.cpp
@@ -0,0 +1,44 @@
+// Copyright Epic Games, Inc. All Rights Reserved.
+
+#include "WorldPartitionRuntimeSpatialHashDetailsCustomization.h"
+#include "WorldPartition/WorldPartitionRuntimeSpatialHash.h"
+#include "DetailLayoutBuilder.h"
+#include "DetailCategoryBuilder.h"
+#include "DetailWidgetRow.h"
+#include "Widgets/Text/STextBlock.h"
+#include "Widgets/Input/SCheckBox.h"
+
+#define LOCTEXT_NAMESPACE "FWorldPartitionRuntimeSpatialHashDetails"
+
+TSharedRef<IDetailCustomization> FWorldPartitionRuntimeSpatialHashDetails::MakeInstance()
+{
+	return MakeShareable(new FWorldPartitionRuntimeSpatialHashDetails);
+}
+
+void FWorldPartitionRuntimeSpatialHashDetails::CustomizeDetails(IDetailLayoutBuilder& DetailBuilder)
+{
+	TArray<TWeakObjectPtr<UObject>> ObjectsBeingCustomized;
+	DetailBuilder.GetObjectsBeingCustomized(ObjectsBeingCustomized);
+	check(ObjectsBeingCustomized.Num() == 1);
+
+	WorldPartitionRuntimeSpatialHash = CastChecked<UWorldPartitionRuntimeSpatialHash>(ObjectsBeingCustomized[0].Get());
+
+	IDetailCategoryBuilder& RuntimeSettingsCategory = DetailBuilder.EditCategory("RuntimeSettings");
+
+	RuntimeSettingsCategory.AddCustomRow(LOCTEXT("PreviewGrids", "Preview Grids"), false)
+		.NameContent()
+		[
+			SNew(STextBlock)
+			.Text(LOCTEXT("WorldPartitionRuntimeSpatialHashPreviewGrids", "Preview Grids"))
+			.ToolTipText(LOCTEXT("WorldPartitionRuntimeSpatialHashPreviewGrids_ToolTip", "Toggle preview grids"))
+			.Font(IDetailLayoutBuilder::GetDetailFont())
+		]
+		.ValueContent()
+		[
+			SNew(SCheckBox)
+			.IsChecked(MakeAttributeLambda([this]() { return WorldPartitionRuntimeSpatialHash->bPreviewGrids ? ECheckBoxState::Checked : ECheckBoxState::Unchecked; }))
+			.OnCheckStateChanged_Lambda([this](ECheckBoxState InState) { WorldPartitionRuntimeSpatialHash->bPreviewGrids = (InState == ECheckBoxState::Checked); })
+		];
+}
+
+#undef LOCTEXT_NAMESPACE
diff --git a/Engine/Source/Editor/WorldPartitionEditor/Private/WorldPartition/Customizations/WorldPartitionRuntimeSpatialHashDetailsCustomization.h b/Engine/Source/Editor/WorldPartitionEditor/Private/WorldPartition/Customizations/WorldPartitionRuntimeSpatialHashDetailsCustomization.h
new file mode 100644
--- /dev/null
+++ b/Engine/Source/Editor/WorldPartitionEditor/Private/WorldPartition/Customizations/WorldPartitionRuntimeSpatialHashDetailsCustomization.h
@@ -0,0 +1,29 @@
+// Copyright Epic Games, Inc. All Rights Reserved.
+
+#pragma once
+
+#include "CoreMinimal.h"
+#include "SlateFwd.h"
+#include "IDetailCustomization.h"
+
+class IDetailLayoutBuilder;
+class UWorldPartitionRuntimeSpatialHash;
+
+enum class ECheckBoxState : uint8;
+
+class FWorldPartitionRuntimeSpatialHashDetails : public IDetailCustomization
+{
+public:
+	/** Makes a new instance of this detail layout class for a specific detail view requesting it */
+	static TSharedRef<IDetailCustomization> MakeInstance();
+
+protected:
+	FWorldPartitionRuntimeSpatialHashDetails()
+	{}
+
+private:
+	/** IDetailCustomization interface */
+	virtual void CustomizeDetails(IDetailLayoutBuilder& DetailBuilder) override;
+
+	UWorldPartitionRuntimeSpatialHash* WorldPartitionRuntimeSpatialHash;
+};
\ No newline at end of file
diff --git a/Engine/Source/Editor/WorldPartitionEditor/Private/WorldPartitionEditorModule.cpp b/Engine/Source/Editor/WorldPartitionEditor/Private/WorldPartitionEditorModule.cpp
--- a/Engine/Source/Editor/WorldPartitionEditor/Private/WorldPartitionEditorModule.cpp
+++ b/Engine/Source/Editor/WorldPartitionEditor/Private/WorldPartitionEditorModule.cpp
@@ -13,6 +13,7 @@
 #include "WorldPartition/SWorldPartitionEditorGridSpatialHash.h"
 
 #include "WorldPartition/Customizations/WorldPartitionDetailsCustomization.h"
+#include "WorldPartition/Customizations/WorldPartitionRuntimeSpatialHashDetailsCustomization.h"
 
 #include "WorldPartition/SWorldPartitionConvertDialog.h"
 #include "WorldPartition/WorldPartitionConvertOptions.h"
@@ -142,6 +143,7 @@ void FWorldPartitionEditorModule::StartupModule()
 
 	FPropertyEditorModule& PropertyEditor = FModuleManager::LoadModuleChecked<FPropertyEditorModule>("PropertyEditor");
 	PropertyEditor.RegisterCustomClassLayout("WorldPartition", FOnGetDetailCustomizationInstance::CreateStatic(&FWorldPartitionDetails::MakeInstance));
+	PropertyEditor.RegisterCustomClassLayout("WorldPartitionRuntimeSpatialHash", FOnGetDetailCustomizationInstance::CreateStatic(&FWorldPartitionRuntimeSpatialHashDetails::MakeInstance));
 }
 
 void FWorldPartitionEditorModule::ShutdownModule()
diff --git a/Engine/Source/Runtime/Engine/Private/WorldPartition/WorldPartitionRuntimeSpatialHash.cpp b/Engine/Source/Runtime/Engine/Private/WorldPartition/WorldPartitionRuntimeSpatialHash.cpp
--- a/Engine/Source/Runtime/Engine/Private/WorldPartition/WorldPartitionRuntimeSpatialHash.cpp
+++ b/Engine/Source/Runtime/Engine/Private/WorldPartition/WorldPartitionRuntimeSpatialHash.cpp
@@ -706,7 +706,7 @@ ASpatialHashRuntimeGridInfo::ASpatialHashRuntimeGridInfo(const FObjectInitialize
 
 UWorldPartitionRuntimeSpatialHash::UWorldPartitionRuntimeSpatialHash(const FObjectInitializer& ObjectInitializer)
 	: Super(ObjectInitializer)
-#if WITH_EDITOR
+#if WITH_EDITORONLY_DATA
 	, bPreviewGrids(false)
 #endif
 {}
diff --git a/Engine/Source/Runtime/Engine/Public/WorldPartition/WorldPartitionRuntimeSpatialHash.h b/Engine/Source/Runtime/Engine/Public/WorldPartition/WorldPartitionRuntimeSpatialHash.h
--- a/Engine/Source/Runtime/Engine/Public/WorldPartition/WorldPartitionRuntimeSpatialHash.h
+++ b/Engine/Source/Runtime/Engine/Public/WorldPartition/WorldPartitionRuntimeSpatialHash.h
@@ -262,12 +262,12 @@ private:
 	UPROPERTY(EditAnywhere, Config, Category = RuntimeSettings)
 	TArray<FSpatialHashRuntimeGrid> Grids;
 
-	UPROPERTY(EditAnywhere, Category = RuntimeSettings, Transient, NonTransactional)
-	bool bPreviewGrids;
-
 	UPROPERTY(Transient)
 	mutable FWorldPartitionRuntimeSpatialHashGridPreviewer GridPreviewer;
 
+	/** Whether to preview runtime grids. */
+	bool bPreviewGrids;
+
 	TMap<FString, UWorldPartitionRuntimeCell*> PackagesToGenerateForCook;
 #endif
 
@@ -294,4 +294,5 @@ private:
 	TArray<const FSpatialHashStreamingGrid*> GetFilteredStreamingGrids() const;
 
 	friend class UWorldPartitionSubsystem;
+	friend class FWorldPartitionRuntimeSpatialHashDetails;
 };
\ No newline at end of file
