The PreviewGrids flag in WorldPartitionRuntimeSpatialHash is now transacted.

#rnx

diff --git a/Engine/Source/Editor/WorldPartitionEditor/Private/WorldPartition/Customizations/WorldPartitionRuntimeSpatialHashDetailsCustomization.cpp b/Engine/Source/Editor/WorldPartitionEditor/Private/WorldPartition/Customizations/WorldPartitionRuntimeSpatialHashDetailsCustomization.cpp
--- a/Engine/Source/Editor/WorldPartitionEditor/Private/WorldPartition/Customizations/WorldPartitionRuntimeSpatialHashDetailsCustomization.cpp
+++ b/Engine/Source/Editor/WorldPartitionEditor/Private/WorldPartition/Customizations/WorldPartitionRuntimeSpatialHashDetailsCustomization.cpp
@@ -7,6 +7,7 @@
 #include "DetailWidgetRow.h"
 #include "Widgets/Text/STextBlock.h"
 #include "Widgets/Input/SCheckBox.h"
+#include "ScopedTransaction.h"
 
 #define LOCTEXT_NAMESPACE "FWorldPartitionRuntimeSpatialHashDetails"
 
@@ -36,8 +37,12 @@ void FWorldPartitionRuntimeSpatialHashDetails::CustomizeDetails(IDetailLayoutBui
 		.ValueContent()
 		[
 			SNew(SCheckBox)
-			.IsChecked(MakeAttributeLambda([this]() { return WorldPartitionRuntimeSpatialHash->bPreviewGrids ? ECheckBoxState::Checked : ECheckBoxState::Unchecked; }))
-			.OnCheckStateChanged_Lambda([this](ECheckBoxState InState) { WorldPartitionRuntimeSpatialHash->bPreviewGrids = (InState == ECheckBoxState::Checked); })
+			.IsChecked(MakeAttributeLambda([this]() { return WorldPartitionRuntimeSpatialHash->GetPreviewGrids() ? ECheckBoxState::Checked : ECheckBoxState::Unchecked; }))
+			.OnCheckStateChanged_Lambda([this](ECheckBoxState InState)
+			{
+				FScopedTransaction Transaction(LOCTEXT("SetRuntimeSpatialHashPreviewGrids", "Change the Runtime Spatial Hash Preview Grids"));
+				WorldPartitionRuntimeSpatialHash->SetPreviewGrids(InState == ECheckBoxState::Checked);
+			})
 		];
 }
 
diff --git a/Engine/Source/Runtime/Engine/Private/WorldPartition/WorldPartitionRuntimeSpatialHash.cpp b/Engine/Source/Runtime/Engine/Private/WorldPartition/WorldPartitionRuntimeSpatialHash.cpp
--- a/Engine/Source/Runtime/Engine/Private/WorldPartition/WorldPartitionRuntimeSpatialHash.cpp
+++ b/Engine/Source/Runtime/Engine/Private/WorldPartition/WorldPartitionRuntimeSpatialHash.cpp
@@ -835,6 +835,17 @@ FName UWorldPartitionRuntimeSpatialHash::GetCellName(UWorldPartition* WorldParti
 	return FName(*FString::Printf(TEXT("%s_%s_%s_DL%X"), *PackageNameNoPIEPrefix, *InGridName.ToString(), *GetCellCoordString(InCellGlobalCoord), InDataLayerID.GetHash()));
 }
 
+bool UWorldPartitionRuntimeSpatialHash::GetPreviewGrids() const
+{
+	return bPreviewGrids;
+}
+
+void UWorldPartitionRuntimeSpatialHash::SetPreviewGrids(bool bInPreviewGrids)
+{
+	Modify(false);
+	bPreviewGrids = bInPreviewGrids;
+}
+
 FName UWorldPartitionRuntimeSpatialHash::GetCellName(FName InGridName, const FIntVector& InCellGlobalCoord, const FDataLayersID& InDataLayerID) const
 {
 	UWorldPartition* WorldPartition = GetOuterUWorldPartition();
diff --git a/Engine/Source/Runtime/Engine/Public/WorldPartition/WorldPartitionRuntimeSpatialHash.h b/Engine/Source/Runtime/Engine/Public/WorldPartition/WorldPartitionRuntimeSpatialHash.h
--- a/Engine/Source/Runtime/Engine/Public/WorldPartition/WorldPartitionRuntimeSpatialHash.h
+++ b/Engine/Source/Runtime/Engine/Public/WorldPartition/WorldPartitionRuntimeSpatialHash.h
@@ -237,6 +237,9 @@ public:
 
 	FName GetCellName(FName InGridName, const FIntVector& InCellGlobalCoord, const FDataLayersID& InDataLayerID) const;
 	static FName GetCellName(UWorldPartition* WorldPartition, FName InGridName, const FIntVector& InCellGlobalCoord, const FDataLayersID& InDataLayerID);
+
+	bool GetPreviewGrids() const;
+	void SetPreviewGrids(bool bInPreviewGrids);
 #endif
 
 	// streaming interface
@@ -262,12 +265,13 @@ private:
 	UPROPERTY(EditAnywhere, Config, Category = RuntimeSettings)
 	TArray<FSpatialHashRuntimeGrid> Grids;
 
-	UPROPERTY(Transient)
-	mutable FWorldPartitionRuntimeSpatialHashGridPreviewer GridPreviewer;
-
 	/** Whether to preview runtime grids. */
+	UPROPERTY(Transient)
 	bool bPreviewGrids;
 
+	UPROPERTY(Transient)
+	mutable FWorldPartitionRuntimeSpatialHashGridPreviewer GridPreviewer;
+
 	TMap<FString, UWorldPartitionRuntimeCell*> PackagesToGenerateForCook;
 #endif
 
@@ -294,5 +298,4 @@ private:
 	TArray<const FSpatialHashStreamingGrid*> GetFilteredStreamingGrids() const;
 
 	friend class UWorldPartitionSubsystem;
-	friend class FWorldPartitionRuntimeSpatialHashDetails;
 };
\ No newline at end of file
