Allow HLOD/Minimap commandlets to be executed with vk
* Must provide a VK project on which to operate

diff --git a/Engine/Source/Editor/UnrealEd/Private/WorldPartition/WorldPartitionBuilder.cpp b/Engine/Source/Editor/UnrealEd/Private/WorldPartition/WorldPartitionBuilder.cpp
--- a/Engine/Source/Editor/UnrealEd/Private/WorldPartition/WorldPartitionBuilder.cpp
+++ b/Engine/Source/Editor/UnrealEd/Private/WorldPartition/WorldPartitionBuilder.cpp
@@ -79,15 +79,24 @@ bool UWorldPartitionBuilder::RunBuilder(UWorld* World)
 	bool bResult = true;
 	// Setup the world
 	{
-		UWorld::InitializationValues IVS;
-		IVS.RequiresHitProxies(false);
-		IVS.ShouldSimulatePhysics(false);
-		IVS.EnableTraceCollision(false);
-		IVS.CreateNavigation(false);
-		IVS.CreateAISystem(false);
-		IVS.AllowAudioPlayback(false);
-		IVS.CreatePhysicsScene(true);
-		FScopedEditorWorld EditorWorld(World, IVS);
+		TUniquePtr<FScopedEditorWorld> EditorWorld;
+		if (World->bIsWorldInitialized)
+		{
+			// Skip initialization, but ensure we're dealing with an editor world.
+			check(World->WorldType == EWorldType::Editor);
+		}
+		else
+		{
+			UWorld::InitializationValues IVS;
+			IVS.RequiresHitProxies(false);
+			IVS.ShouldSimulatePhysics(false);
+			IVS.EnableTraceCollision(false);
+			IVS.CreateNavigation(false);
+			IVS.CreateAISystem(false);
+			IVS.AllowAudioPlayback(false);
+			IVS.CreatePhysicsScene(true);
+			EditorWorld = MakeUnique<FScopedEditorWorld>(World, IVS);
+		}
 
 		// Make sure the world is partitioned
 		if (UWorld::IsPartitionedWorld(World))
diff --git a/Engine/Source/Editor/WorldPartitionEditor/Private/WorldPartitionEditorModule.cpp b/Engine/Source/Editor/WorldPartitionEditor/Private/WorldPartitionEditorModule.cpp
--- a/Engine/Source/Editor/WorldPartitionEditor/Private/WorldPartitionEditorModule.cpp
+++ b/Engine/Source/Editor/WorldPartitionEditor/Private/WorldPartitionEditorModule.cpp
@@ -259,7 +259,7 @@ static void RescanAssetsAndLoadMap(const FString& MapToLoad)
 	FEditorFileUtils::LoadMap(MapToLoad);
 }
 
-static void RunCommandletAsExternalProcess(const FString& InCommandletArgs, const FText& InOperationDescription, int32& OutResult, bool& bOutCancelled, FString& OutCommandletOutput)
+void FWorldPartitionEditorModule::RunCommandletAsExternalProcess(const FString& InCommandletArgs, const FText& InOperationDescription, int32& OutResult, bool& bOutCancelled, FString& OutCommandletOutput)
 {
 	OutResult = 0;
 	bOutCancelled = false;
@@ -276,13 +276,26 @@ static void RunCommandletAsExternalProcess(const FString& InCommandletArgs, cons
 
 	FString CurrentExecutableName = FPlatformProcess::ExecutablePath();
 
+	TArray<FString> AdditionalArgs;
+	OnExecuteCommandletEvent.Broadcast(AdditionalArgs);
+
 	// Try to provide complete Path, if we can't try with project name
 	FString ProjectPath = FPaths::IsProjectFilePathSet() ? FPaths::GetProjectFilePath() : FApp::GetProjectName();
 
-	uint32 ProcessID;
-	FString Arguments = FString::Printf(TEXT("\"%s\" %s -unattended"), *ProjectPath, *InCommandletArgs);
+	FString Arguments;
+	Arguments += ProjectPath;
+	Arguments += TEXT(" -BaseDir=\"") + FString(FPlatformProcess::BaseDir()) + TEXT('"');
+	Arguments += TEXT(" -Unattended");
+	Arguments += TEXT(" ") + InCommandletArgs;
+	for (const FString& AdditionalArg : AdditionalArgs)
+	{
+		Arguments += " ";
+		Arguments += AdditionalArg;
+	}
 	
 	UE_LOG(LogWorldPartitionEditor, Display, TEXT("Running commandlet: %s %s"), *CurrentExecutableName, *Arguments);
+
+	uint32 ProcessID;
 	const bool bLaunchDetached = true;
 	const bool bLaunchHidden = true;
 	const bool bLaunchReallyHidden = true;
@@ -373,8 +386,6 @@ bool FWorldPartitionEditorModule::ConvertMap(const FString& InLongPackageName)
 		RunCommandletAsExternalProcess(CommandletArgs, OperationDescription, Result, bCancelled, CommandletOutput);
 		if (!bCancelled && Result == 0)
 		{	
-			FMessageDialog::Open(EAppMsgType::Ok, FText::Format(LOCTEXT("ConvertMapCompleted", "Conversion completed:\n{0}"), FText::FromString(CommandletOutput)));
-
 #if	PLATFORM_DESKTOP
 			if (DefaultConvertOptions->bGenerateIni)
 			{
@@ -398,7 +409,7 @@ bool FWorldPartitionEditorModule::ConvertMap(const FString& InLongPackageName)
 		}
 		else if(Result != 0)
 		{
-			FMessageDialog::Open(EAppMsgType::Ok, FText::Format(LOCTEXT("ConvertMapFailed", "Conversion failed:\n{0}"), FText::FromString(CommandletOutput)));
+			FMessageDialog::Open(EAppMsgType::Ok, LOCTEXT("ConvertMapFailed", "Conversion failed! See log for details."));
 		}
 	}
 
@@ -461,7 +472,6 @@ bool FWorldPartitionEditorModule::BuildHLODs(const FString& InMapToProcess)
 		bool bSuccess = !bCancelled && Result == 0;
 		if (bSuccess)
 		{
-			FMessageDialog::Open(EAppMsgType::Ok, FText::Format(LOCTEXT("HLODBuildCompleted", "HLOD build completed:\n{0}"), FText::FromString(CommandletOutput)));
 			RescanAssetsAndLoadMap(InMapToProcess);
 		}
 		else if (bCancelled)
@@ -470,7 +480,7 @@ bool FWorldPartitionEditorModule::BuildHLODs(const FString& InMapToProcess)
 		}
 		else if (Result != 0)
 		{
-			FMessageDialog::Open(EAppMsgType::Ok, FText::Format(LOCTEXT("HLODBuildFailed", "HLOD build failed:\n{0}"), FText::FromString(CommandletOutput)));
+			FMessageDialog::Open(EAppMsgType::Ok, LOCTEXT("HLODBuildFailed", "HLOD build failed! See log for details."));
 		}
 
 		return bSuccess;
@@ -497,7 +507,6 @@ bool FWorldPartitionEditorModule::BuildMinimap(const FString& InMapToProcess)
 	bool bSuccess = !bCancelled && Result == 0;
 	if (bSuccess)
 	{
-		FMessageDialog::Open(EAppMsgType::Ok, FText::Format(LOCTEXT("MinimapBuildCompleted", "Minimap build completed:\n{0}"), FText::FromString(CommandletOutput)));
 		RescanAssetsAndLoadMap(InMapToProcess);
 	}
 	else if (bCancelled)
@@ -506,7 +515,7 @@ bool FWorldPartitionEditorModule::BuildMinimap(const FString& InMapToProcess)
 	}
 	else if (Result != 0)
 	{
-		FMessageDialog::Open(EAppMsgType::Ok, FText::Format(LOCTEXT("MinimapBuildFailed", "Minimap build failed:\n{0}"), FText::FromString(CommandletOutput)));
+		FMessageDialog::Open(EAppMsgType::Ok, LOCTEXT("MinimapBuildFailed", "Minimap build failed! See log for details."));
 	}
 
 	return bSuccess;
diff --git a/Engine/Source/Editor/WorldPartitionEditor/Public/WorldPartitionEditorModule.h b/Engine/Source/Editor/WorldPartitionEditor/Public/WorldPartitionEditorModule.h
--- a/Engine/Source/Editor/WorldPartitionEditor/Public/WorldPartitionEditorModule.h
+++ b/Engine/Source/Editor/WorldPartitionEditor/Public/WorldPartitionEditorModule.h
@@ -51,11 +51,12 @@ public:
 	 */
 	virtual bool RunBuilder(TSubclassOf<UWorldPartitionBuilder> WorldPartitionBuilder, const FString& InLongPackageName) override;
 
-	/**
-	 *
-	 */
+	/** Return the world added event. */
 	virtual FWorldPartitionCreated& OnWorldPartitionCreated() override { return WorldPartitionCreatedEvent; }
 
+	/** Return the commandlet execution event */
+	virtual FOnExecuteCommandlet& OnExecuteCommandlet() override { return OnExecuteCommandletEvent; }
+
 private:
 	/** Called when the level editors map changes. We will determine if the new map is a valid world partition world and close world partition tabs if not */
 	void OnMapChanged(uint32 MapFlags);
@@ -73,6 +74,7 @@ private:
 	bool BuildMinimap(const FString& InMapToProcess);
 
 private:
+	void RunCommandletAsExternalProcess(const FString& InCommandletArgs, const FText& InOperationDescription, int32& OutResult, bool& bOutCancelled, FString& OutCommandletOutput);
 	void OnConvertMap();
 
 	FDelegateHandle LevelEditorExtenderDelegateHandle;
@@ -82,4 +84,6 @@ private:
 	TWeakPtr<SDockTab> WorldPartitionTab;
 
 	FWorldPartitionCreated WorldPartitionCreatedEvent;
+
+	FOnExecuteCommandlet OnExecuteCommandletEvent;
 };
diff --git a/Engine/Source/Runtime/Engine/Public/WorldPartition/IWorldPartitionEditorModule.h b/Engine/Source/Runtime/Engine/Public/WorldPartition/IWorldPartitionEditorModule.h
--- a/Engine/Source/Runtime/Engine/Public/WorldPartition/IWorldPartitionEditorModule.h
+++ b/Engine/Source/Runtime/Engine/Public/WorldPartition/IWorldPartitionEditorModule.h
@@ -25,6 +25,12 @@ public:
 	/** Triggered when a world is added. */
 	DECLARE_EVENT_OneParam(IWorldPartitionEditorModule, FWorldPartitionCreated, UWorld*);
 
+	/** Triggered when the editor launches a commandlet. Can be used to provide project specific arguments. */
+	DECLARE_EVENT_OneParam(IWorldPartitionEditorModule, FOnExecuteCommandlet, TArray<FString>&);
+
 	/** Return the world added event. */
 	virtual FWorldPartitionCreated& OnWorldPartitionCreated() = 0;
+
+	/** Return the commandlet execution event */
+	virtual FOnExecuteCommandlet& OnExecuteCommandlet() = 0;
 };
\ No newline at end of file
