Fixed issue with WorldPartitionEditorModule when running from a commandlet

diff --git a/Engine/Source/Editor/WorldPartitionEditor/Private/WorldPartitionEditorModule.cpp b/Engine/Source/Editor/WorldPartitionEditor/Private/WorldPartitionEditorModule.cpp
--- a/Engine/Source/Editor/WorldPartitionEditor/Private/WorldPartitionEditorModule.cpp
+++ b/Engine/Source/Editor/WorldPartitionEditor/Private/WorldPartitionEditorModule.cpp
@@ -111,31 +111,7 @@ void FWorldPartitionEditorModule::StartupModule()
 	SWorldPartitionEditorGrid::RegisterPartitionEditorGridCreateInstanceFunc(NAME_None, &SWorldPartitionEditorGrid::CreateInstance);
 	SWorldPartitionEditorGrid::RegisterPartitionEditorGridCreateInstanceFunc(TEXT("SpatialHash"), &SWorldPartitionEditorGridSpatialHash::CreateInstance);
 	
-	if (!IsRunningGame())
-	{
-		FLevelEditorModule& LevelEditorModule = FModuleManager::Get().LoadModuleChecked<FLevelEditorModule>("LevelEditor");
-		TArray<FLevelEditorModule::FLevelViewportMenuExtender_SelectedActors>& MenuExtenderDelegates = LevelEditorModule.GetAllLevelViewportContextMenuExtenders();
-
-
-		LevelEditorModule.OnRegisterTabs().AddRaw(this, &FWorldPartitionEditorModule::RegisterWorldPartitionTabs);
-		LevelEditorModule.OnRegisterLayoutExtensions().AddRaw(this, &FWorldPartitionEditorModule::RegisterWorldPartitionLayout);
-
-		MenuExtenderDelegates.Add(FLevelEditorModule::FLevelViewportMenuExtender_SelectedActors::CreateStatic(&OnExtendLevelEditorMenu));
-		LevelEditorExtenderDelegateHandle = MenuExtenderDelegates.Last().GetHandle();
-
-		FToolMenuOwnerScoped OwnerScoped(this);
-		UToolMenu* Menu = UToolMenus::Get()->ExtendMenu("LevelEditor.MainMenu.Tools");
-		FToolMenuSection& Section = Menu->AddSection("WorldPartition", LOCTEXT("WorldPartition", "World Partition"));
-		Section.AddEntry(FToolMenuEntry::InitMenuEntry(
-			"WorldPartition",
-			LOCTEXT("WorldPartitionConvertTitle", "Convert Level..."),
-			LOCTEXT("WorldPartitionConvertTooltip", "Converts a Level to World Partition."),
-			FSlateIcon(FEditorStyle::GetStyleSetName(), "DeveloperTools.MenuIcon"),
-			FUIAction(FExecuteAction::CreateRaw(this, &FWorldPartitionEditorModule::OnConvertMap))
-		));
-
-		FEditorDelegates::MapChange.AddRaw(this, &FWorldPartitionEditorModule::OnMapChanged);
-	}
+	UToolMenus::RegisterStartupCallback(FSimpleMulticastDelegate::FDelegate::CreateRaw(this, &FWorldPartitionEditorModule::RegisterMenus));	
 
 	IAssetTools& AssetTools = FModuleManager::LoadModuleChecked<FAssetToolsModule>("AssetTools").Get();
 	HLODLayerAssetTypeActions = MakeShareable(new FHLODLayerAssetTypeActions);
@@ -166,6 +142,7 @@ void FWorldPartitionEditorModule::ShutdownModule()
 
 		FEditorDelegates::MapChange.RemoveAll(this);
 
+		UToolMenus::UnRegisterStartupCallback(this);
 		UToolMenus::UnregisterOwner(this);
 	}
 
@@ -187,6 +164,31 @@ void FWorldPartitionEditorModule::ShutdownModule()
 	}
 }
 
+void FWorldPartitionEditorModule::RegisterMenus()
+{
+	FLevelEditorModule& LevelEditorModule = FModuleManager::Get().LoadModuleChecked<FLevelEditorModule>("LevelEditor");
+	TArray<FLevelEditorModule::FLevelViewportMenuExtender_SelectedActors>& MenuExtenderDelegates = LevelEditorModule.GetAllLevelViewportContextMenuExtenders();
+
+	LevelEditorModule.OnRegisterTabs().AddRaw(this, &FWorldPartitionEditorModule::RegisterWorldPartitionTabs);
+	LevelEditorModule.OnRegisterLayoutExtensions().AddRaw(this, &FWorldPartitionEditorModule::RegisterWorldPartitionLayout);
+
+	MenuExtenderDelegates.Add(FLevelEditorModule::FLevelViewportMenuExtender_SelectedActors::CreateStatic(&OnExtendLevelEditorMenu));
+	LevelEditorExtenderDelegateHandle = MenuExtenderDelegates.Last().GetHandle();
+
+	FToolMenuOwnerScoped OwnerScoped(this);
+	UToolMenu* Menu = UToolMenus::Get()->ExtendMenu("LevelEditor.MainMenu.Tools");
+	FToolMenuSection& Section = Menu->AddSection("WorldPartition", LOCTEXT("WorldPartition", "World Partition"));
+	Section.AddEntry(FToolMenuEntry::InitMenuEntry(
+		"WorldPartition",
+		LOCTEXT("WorldPartitionConvertTitle", "Convert Level..."),
+		LOCTEXT("WorldPartitionConvertTooltip", "Converts a Level to World Partition."),
+		FSlateIcon(FEditorStyle::GetStyleSetName(), "DeveloperTools.MenuIcon"),
+		FUIAction(FExecuteAction::CreateRaw(this, &FWorldPartitionEditorModule::OnConvertMap))
+	));
+
+	FEditorDelegates::MapChange.AddRaw(this, &FWorldPartitionEditorModule::OnMapChanged);
+}
+
 TSharedRef<SWidget> FWorldPartitionEditorModule::CreateWorldPartitionEditor()
 {
 	UWorld* EditorWorld = GEditor->GetEditorWorldContext().World();
diff --git a/Engine/Source/Editor/WorldPartitionEditor/Public/WorldPartitionEditorModule.h b/Engine/Source/Editor/WorldPartitionEditor/Public/WorldPartitionEditorModule.h
--- a/Engine/Source/Editor/WorldPartitionEditor/Public/WorldPartitionEditorModule.h
+++ b/Engine/Source/Editor/WorldPartitionEditor/Public/WorldPartitionEditorModule.h
@@ -61,6 +61,9 @@ private:
 	/** Called when the level editors map changes. We will determine if the new map is a valid world partition world and close world partition tabs if not */
 	void OnMapChanged(uint32 MapFlags);
 
+	/** Register menus */
+	void RegisterMenus();
+
 	/** Registers world partition tabs spawners with the level editor */
 	void RegisterWorldPartitionTabs(TSharedPtr<FTabManager> InTabManager);
 
