Fixed issues when running HLOD/Minimap builds on an unsaved map
* External commandlet would be executed on the /Temp/Untitled package

diff --git a/Engine/Source/Editor/UnrealEd/Private/EditorBuildUtils.cpp b/Engine/Source/Editor/UnrealEd/Private/EditorBuildUtils.cpp
--- a/Engine/Source/Editor/UnrealEd/Private/EditorBuildUtils.cpp
+++ b/Engine/Source/Editor/UnrealEd/Private/EditorBuildUtils.cpp
@@ -499,8 +499,7 @@ bool FEditorBuildUtils::EditorBuild( UWorld* InWorld, FName Id, const bool bAllo
 
 	// It's possible the world was unloaded & reloaded if external commands were run.
 	// To work around this, reassign the initial world from it's soft object path.
-	check(World.IsValid());
-	InWorld = World.Get();
+	InWorld = World.IsValid() ? World.Get() : GEditor->GetEditorWorldContext().World();
 
 	// Check map for errors (only if build operation happened)
 	if ( bShouldMapCheck && bDoBuild && !GEditor->GetMapBuildCancelled() )
@@ -1169,8 +1168,7 @@ void FBuildAllHandler::ProcessBuild(const TWeakPtr<SBuildProgressWidget>& BuildP
 	{
 		// It's possible the world was unloaded & reloaded if external commands were run.
 		// To work around this, reassign the initial world from it's soft object path.
-		check(World.IsValid());
-		CurrentWorld = World.Get();
+		CurrentWorld = World.IsValid() ? World.Get() : GEditor->GetEditorWorldContext().World();
 
 		if (GEditor->GetMapBuildCancelled())
 		{
diff --git a/Engine/Source/Editor/WorldPartitionEditor/Private/WorldPartitionEditorModule.cpp b/Engine/Source/Editor/WorldPartitionEditor/Private/WorldPartitionEditorModule.cpp
--- a/Engine/Source/Editor/WorldPartitionEditor/Private/WorldPartitionEditorModule.cpp
+++ b/Engine/Source/Editor/WorldPartitionEditor/Private/WorldPartitionEditorModule.cpp
@@ -226,14 +226,22 @@ void FWorldPartitionEditorModule::OnConvertMap()
 	}
 }
 
-static bool UnloadCurrentMap(bool bAskSaveContentPackages)
+static bool UnloadCurrentMap(bool bAskSaveContentPackages, FString& MapPackageName)
 {
+	UPackage* WorldPackage = FindPackage(nullptr, *MapPackageName);
+
 	// Ask user to save dirty packages
 	if (!FEditorFileUtils::SaveDirtyPackages(/*bPromptUserToSave=*/true, /*bSaveMapPackages=*/true, bAskSaveContentPackages))
 	{
 		return false;
 	}
 
+	// Make sure we handle the case where the world package was renamed on save (for temp world for example)
+	if (WorldPackage)
+	{
+		MapPackageName = WorldPackage->GetLoadedPath().GetPackageName();
+	}
+
 	// Unload any loaded map
 	if (!UEditorLoadingAndSavingUtils::NewBlankMap(/*bSaveExistingMap*/false))
 	{
@@ -259,6 +267,12 @@ static void RescanAssetsAndLoadMap(const FString& MapToLoad)
 	AssetRegistry.ScanPathsSynchronous(ExternalObjectsPaths, true);
 
 	FEditorFileUtils::LoadMap(MapToLoad);
+
+	UWorld* World = GEditor->GetEditorWorldContext().World();
+	if (!World || World->GetPackage()->GetLoadedPath().GetPackageName() != MapToLoad)
+	{
+		UE_LOG(LogWorldPartitionEditor, Error, TEXT("Failed to reopen world."));
+	}
 }
 
 void FWorldPartitionEditorModule::RunCommandletAsExternalProcess(const FString& InCommandletArgs, const FText& InOperationDescription, int32& OutResult, bool& bOutCancelled, FString& OutCommandletOutput)
@@ -374,7 +388,7 @@ bool FWorldPartitionEditorModule::ConvertMap(const FString& InLongPackageName)
 
 	if (ConvertDialog->ClickedOk())
 	{
-		if (!UnloadCurrentMap(/*bAskSaveContentPackages=*/false))
+		if (!UnloadCurrentMap(/*bAskSaveContentPackages=*/false, DefaultConvertOptions->LongPackageName))
 		{
 			return false;
 		}
@@ -397,7 +411,7 @@ bool FWorldPartitionEditorModule::ConvertMap(const FString& InLongPackageName)
 			}
 #endif				
 				
-			FString MapToLoad = InLongPackageName;
+			FString MapToLoad = DefaultConvertOptions->LongPackageName;
 			if (!DefaultConvertOptions->bInPlace)
 			{
 				MapToLoad += UWorldPartitionConvertCommandlet::GetConversionSuffix(DefaultConvertOptions->bOnlyMergeSubLevels);
@@ -457,13 +471,14 @@ bool FWorldPartitionEditorModule::BuildHLODs(const FString& InMapToProcess)
 
 	if (BuildHLODsDialog->GetDialogResult() != SWorldPartitionBuildHLODsDialog::DialogResult::Cancel)
 	{
-		if (!UnloadCurrentMap(/*bAskSaveContentPackages=*/true))
+		FString MapPackage = InMapToProcess;
+		if (!UnloadCurrentMap(/*bAskSaveContentPackages=*/true, MapPackage))
 		{
 			return false;
 		}
 
 		const FString BuildArgs = BuildHLODsDialog->GetDialogResult() == SWorldPartitionBuildHLODsDialog::DialogResult::BuildHLODs ? "-SetupHLODs -BuildHLODs -AllowCommandletRendering" : "-DeleteHLODs";
-		const FString CommandletArgs = InMapToProcess + " -run=WorldPartitionBuilderCommandlet -Builder=WorldPartitionHLODsBuilder " + BuildArgs;
+		const FString CommandletArgs = MapPackage + " -run=WorldPartitionBuilderCommandlet -Builder=WorldPartitionHLODsBuilder " + BuildArgs;
 		const FText OperationDescription = LOCTEXT("HLODBuildProgress", "Building HLODs...");
 
 		int32 Result;
@@ -474,7 +489,7 @@ bool FWorldPartitionEditorModule::BuildHLODs(const FString& InMapToProcess)
 		bool bSuccess = !bCancelled && Result == 0;
 		if (bSuccess)
 		{
-			RescanAssetsAndLoadMap(InMapToProcess);
+			RescanAssetsAndLoadMap(MapPackage);
 		}
 		else if (bCancelled)
 		{
@@ -493,12 +508,13 @@ bool FWorldPartitionEditorModule::BuildHLODs(const FString& InMapToProcess)
 
 bool FWorldPartitionEditorModule::BuildMinimap(const FString& InMapToProcess)
 {
-	if (!UnloadCurrentMap(/*bAskSaveContentPackages=*/true))
+	FString MapPackage = InMapToProcess;
+	if (!UnloadCurrentMap(/*bAskSaveContentPackages=*/true, MapPackage))
 	{
 		return false;
 	}
 
-	const FString CommandletArgs = InMapToProcess + " -run=WorldPartitionBuilderCommandlet -Builder=WorldPartitionMinimapBuilder -AllowCommandletRendering";
+	const FString CommandletArgs = MapPackage + " -run=WorldPartitionBuilderCommandlet -Builder=WorldPartitionMinimapBuilder -AllowCommandletRendering";
 	const FText OperationDescription = LOCTEXT("MinimapBuildProgress", "Building minimap...");
 
 	int32 Result;
@@ -509,7 +525,7 @@ bool FWorldPartitionEditorModule::BuildMinimap(const FString& InMapToProcess)
 	bool bSuccess = !bCancelled && Result == 0;
 	if (bSuccess)
 	{
-		RescanAssetsAndLoadMap(InMapToProcess);
+		RescanAssetsAndLoadMap(MapPackage);
 	}
 	else if (bCancelled)
 	{
