Fixed crash when loading/switching from 2 partitioned maps in the editor : WorldPartitionDetailsCustomizaiton now keeps a WeakPtr on the WorldPartition object.
#rnx

diff --git a/Engine/Source/Editor/WorldPartitionEditor/Private/WorldPartition/Customizations/WorldPartitionDetailsCustomization.cpp b/Engine/Source/Editor/WorldPartitionEditor/Private/WorldPartition/Customizations/WorldPartitionDetailsCustomization.cpp
--- a/Engine/Source/Editor/WorldPartitionEditor/Private/WorldPartition/Customizations/WorldPartitionDetailsCustomization.cpp
+++ b/Engine/Source/Editor/WorldPartitionEditor/Private/WorldPartition/Customizations/WorldPartitionDetailsCustomization.cpp
@@ -42,10 +42,10 @@ void FWorldPartitionDetails::CustomizeDetails(IDetailLayoutBuilder& DetailBuilde
 		.ValueContent()
 		[
 			SNew(SCheckBox)
-			.IsChecked(MakeAttributeLambda([this]() { return WorldPartition->IsStreamingEnabled() ? ECheckBoxState::Checked : ECheckBoxState::Unchecked; }))
+			.IsChecked(MakeAttributeLambda([this]() { return WorldPartition.IsValid() && WorldPartition->IsStreamingEnabled() ? ECheckBoxState::Checked : ECheckBoxState::Unchecked; }))
 			.OnCheckStateChanged(this, &FWorldPartitionDetails::HandleWorldPartitionEnableStreamingChanged)
 		]
-		.Visibility(TAttribute<EVisibility>::CreateLambda([this]() { return WorldPartition->SupportsStreaming() ? EVisibility::Visible : EVisibility::Hidden; }));
+		.Visibility(TAttribute<EVisibility>::CreateLambda([this]() { return WorldPartition.IsValid() && WorldPartition->SupportsStreaming() ? EVisibility::Visible : EVisibility::Hidden; }));
 
 	if (WorldPartition->RuntimeHash)
 	{
@@ -58,7 +58,7 @@ void FWorldPartitionDetails::CustomizeDetails(IDetailLayoutBuilder& DetailBuilde
 			(*RuntimeHashRow)
 				.ShouldAutoExpand(true)
 				.DisplayName(LOCTEXT("RuntimeHash", "Runtime Hash"))
-				.Visibility(TAttribute<EVisibility>::CreateLambda([this]() { return WorldPartition->IsStreamingEnabled() ? EVisibility::Visible : EVisibility::Hidden; }));
+				.Visibility(TAttribute<EVisibility>::CreateLambda([this]() { return WorldPartition.IsValid() && WorldPartition->IsStreamingEnabled() ? EVisibility::Visible : EVisibility::Hidden; }));
 		}
 	}
 
diff --git a/Engine/Source/Editor/WorldPartitionEditor/Private/WorldPartition/Customizations/WorldPartitionDetailsCustomization.h b/Engine/Source/Editor/WorldPartitionEditor/Private/WorldPartition/Customizations/WorldPartitionDetailsCustomization.h
--- a/Engine/Source/Editor/WorldPartitionEditor/Private/WorldPartition/Customizations/WorldPartitionDetailsCustomization.h
+++ b/Engine/Source/Editor/WorldPartitionEditor/Private/WorldPartition/Customizations/WorldPartitionDetailsCustomization.h
@@ -5,9 +5,9 @@
 #include "CoreMinimal.h"
 #include "SlateFwd.h"
 #include "IDetailCustomization.h"
+#include "WorldPartition/WorldPartition.h"
 
 class IDetailLayoutBuilder;
-class UWorldPartition;
 
 enum class ECheckBoxState : uint8;
 
@@ -34,5 +34,5 @@ private:
 	// Callback for getting the world partition editor cell size.
 	TOptional<uint32> HandleWorldPartitionEditorCellSizeValue() const;
 
-	UWorldPartition* WorldPartition;
+	TWeakObjectPtr<UWorldPartition> WorldPartition;
 };
\ No newline at end of file
