World Partition Editor Hash:
- Unbound octree now stores hierarchical editor cells instead of leaves, this greatly reduce the number of needed editor cells for very large worlds.
- Implemented loading regions support to replace editor cells loading, in preparation for bookmark regions.

Tested against a 2000KMx2000km world + various internal projects and demos.

diff --git a/Engine/Source/Editor/DataLayerEditor/Private/DataLayer/DataLayerTransaction.h b/Engine/Source/Editor/DataLayerEditor/Private/DataLayer/DataLayerTransaction.h
--- a/Engine/Source/Editor/DataLayerEditor/Private/DataLayer/DataLayerTransaction.h
+++ b/Engine/Source/Editor/DataLayerEditor/Private/DataLayer/DataLayerTransaction.h
@@ -3,49 +3,14 @@
 #pragma once
 
 #include "ScopedTransaction.h"
-#include "WorldPartition/WorldPartition.h"
-#include "Engine/World.h"
-#include "Editor.h"
 
 class FScopedDataLayerTransaction
 {
 public:
-	FScopedDataLayerTransaction(const FText& SessionName, UWorld* InWorld)
-	: WorldPtr(InWorld)
-	, bUndoTransaction(false)
-	{
-		check(InWorld);
-		ScopedTransaction = new FScopedTransaction(SessionName);
-		UWorldPartition* WorldPartition = InWorld->GetWorldPartition();
-		WorldPartition->OnCancelWorldPartitionUpdateEditorCells.AddRaw(this, &FScopedDataLayerTransaction::OnCancelUpdateEditorCells);
-	}
-
-	~FScopedDataLayerTransaction()
-	{
-		UWorld* World = WorldPtr.Get();
-		if (UWorldPartition* WorldPartition = World ? World->GetWorldPartition() : nullptr)
-		{
-			WorldPartition->OnCancelWorldPartitionUpdateEditorCells.RemoveAll(this);
-		}
-		delete ScopedTransaction;
-		if (bUndoTransaction)
-		{
-			GEditor->UndoTransaction();
-		}
-	}
+	FScopedDataLayerTransaction(const FText& SessionName, UWorld*)
+		: ScopedTransaction(SessionName)
+	{}
 
 private:
-
-	void OnCancelUpdateEditorCells(UWorldPartition* InWorldPartition)
-	{
-		UWorld* World = WorldPtr.Get();
-		if (World && (World->GetWorldPartition() == InWorldPartition))
-		{
-			bUndoTransaction = true;
-		}
-	}
-
-	FScopedTransaction* ScopedTransaction;
-	TWeakObjectPtr<UWorld> WorldPtr;
-	bool bUndoTransaction;
+	FScopedTransaction ScopedTransaction;
 };
\ No newline at end of file
diff --git a/Engine/Source/Editor/FoliageEdit/Private/ProceduralFoliageComponentDetails.cpp b/Engine/Source/Editor/FoliageEdit/Private/ProceduralFoliageComponentDetails.cpp
--- a/Engine/Source/Editor/FoliageEdit/Private/ProceduralFoliageComponentDetails.cpp
+++ b/Engine/Source/Editor/FoliageEdit/Private/ProceduralFoliageComponentDetails.cpp
@@ -113,13 +113,7 @@ FReply FProceduralFoliageComponentDetails::OnLoadUnloadedAreas()
 	{
 		if(Component.IsValid())
 		{
-			if (UWorldPartition* WorldPartition = Component->GetWorld()->GetWorldPartition())
-			{
-				FVector Origin;
-				FVector Extent;
-				Component->GetOwner()->GetActorBounds(false, Origin, Extent);
-				WorldPartition->LoadEditorCells(FBox(Origin - Extent, Origin + Extent), true);
-			}
+			Component->LoadSimulatedRegion();
 		}
 	}
 
@@ -190,16 +184,9 @@ bool FProceduralFoliageComponentDetails::HasUnloadedAreas() const
 	{
 		if(Component.IsValid())
 		{
-			if (UWorldPartition* WorldPartition = Component->GetWorld()->GetWorldPartition())
+			if (!Component->IsSimulatedRegionLoaded())
 			{
-				FVector Origin;
-				FVector Extent;
-				Component->GetOwner()->GetActorBounds(false, Origin, Extent);
-	
-				if (!WorldPartition->AreEditorCellsLoaded(FBox(Origin - Extent, Origin + Extent)))
-				{
-					return true;
-				}
+				return true;
 			}
 		}
 	}
diff --git a/Engine/Source/Editor/UnrealEd/Private/FileHelpers.cpp b/Engine/Source/Editor/UnrealEd/Private/FileHelpers.cpp
--- a/Engine/Source/Editor/UnrealEd/Private/FileHelpers.cpp
+++ b/Engine/Source/Editor/UnrealEd/Private/FileHelpers.cpp
@@ -729,8 +729,8 @@ static bool SaveWorld(UWorld* World,
 		UWorldPartition* RenamedWorldPartition = nullptr;
 		TArray<FWorldPartitionReference> ActorReferences;
 
-		// Save Loaded cells
-		TArray<FName> LoadedEditorCells;
+		// Save loaded regions
+		TArray<FBox> LoadedEditorRegions;
 
 		// Other packages to save
 		TArray<UPackage*> PackagesToSave;
@@ -754,7 +754,7 @@ static bool SaveWorld(UWorld* World,
 				// When creating a Partitioned Level for a Level Instance, the WorldPartition is not initialized, so no need to do this.
 				if (RenamedWorldPartition && RenamedWorldPartition->IsInitialized())
 				{
-					LoadedEditorCells = RenamedWorldPartition->GetUserLoadedEditorGridCells();
+					LoadedEditorRegions = RenamedWorldPartition->GetUserLoadedEditorGridRegions();
 					RenamedWorldPartition->LoadAllActors(ActorReferences);
 
 					if (bIsTempPackage)
@@ -899,9 +899,9 @@ static bool SaveWorld(UWorld* World,
 
 			if (RenamedWorldPartition)
 			{
-				// Save Snapshot of loaded Editor Cells
-				GetMutableDefault<UWorldPartitionEditorPerProjectUserSettings>()->SetEditorGridLoadedCells(SaveWorld, LoadedEditorCells);
-				RenamedWorldPartition->LoadEditorCells(LoadedEditorCells, /*bIsFromUserChange=*/true);
+				// Save Snapshot of loaded Editor regions
+				GetMutableDefault<UWorldPartitionEditorPerProjectUserSettings>()->SetEditorGridLoadedRegions(SaveWorld, LoadedEditorRegions);
+				RenamedWorldPartition->LoadLastLoadedRegions();
 			}
 		}
 
diff --git a/Engine/Source/Editor/UnrealEd/Private/WorldPartition/WorldPartitionBuilder.cpp b/Engine/Source/Editor/UnrealEd/Private/WorldPartition/WorldPartitionBuilder.cpp
--- a/Engine/Source/Editor/UnrealEd/Private/WorldPartition/WorldPartitionBuilder.cpp
+++ b/Engine/Source/Editor/UnrealEd/Private/WorldPartition/WorldPartitionBuilder.cpp
@@ -14,6 +14,7 @@
 #include "WorldPartition/WorldPartitionSubsystem.h"
 #include "WorldPartition/DataLayer/DataLayerSubsystem.h"
 #include "WorldPartition/DataLayer/DataLayerInstance.h"
+#include "WorldPartition/LoaderAdapter/LoaderAdapterShape.h"
 #include "LevelInstance/LevelInstanceSubsystem.h"
 #include "Math/IntVector.h"
 #include "UObject/SavePackage.h"
@@ -233,7 +234,7 @@ bool UWorldPartitionBuilder::Run(UWorld* World, FPackageSourceControlHelper& Pac
 		UE_LOG(LogWorldPartitionBuilder, Display, TEXT("WorldBounds: Min %s, Max %s"), *CellInfo.EditorBounds.Min.ToString(), *CellInfo.EditorBounds.Max.ToString());
 		UE_LOG(LogWorldPartitionBuilder, Display, TEXT("Iteration Count: %d"), IterationCount);
 		
-		FBox LoadedBounds(ForceInit);
+		TArray<TUniquePtr<IWorldPartitionActorLoaderInterface::ILoaderAdapter>> LoaderAdapters;
 
 		for (int32 z = BeginCellCoords.Z; CanIterateZ(bResult, z); z++)
 		{
@@ -260,17 +261,14 @@ bool UWorldPartitionBuilder::Run(UWorld* World, FPackageSourceControlHelper& Pac
 					CellInfo.Bounds = BoundsToLoad;
 
 					UE_LOG(LogWorldPartitionBuilder, Verbose, TEXT("Loading Bounds: Min %s, Max %s"), *BoundsToLoad.Min.ToString(), *BoundsToLoad.Max.ToString());
-					WorldPartition->LoadEditorCells(BoundsToLoad, false);
-					LoadedBounds += BoundsToLoad;
+
+					LoaderAdapters.Emplace_GetRef(MakeUnique<FLoaderAdapterShape>(World, BoundsToLoad, TEXT("Loaded Region")))->Load();
 
 					bResult = RunInternal(World, CellInfo, PackageHelper);
 
 					if (FWorldPartitionHelpers::HasExceededMaxMemory())
 					{
-						WorldPartition->UnloadEditorCells(LoadedBounds, false);
-						// Reset Loaded Bounds
-						LoadedBounds.Init();
-
+						LoaderAdapters.Empty();
 						FWorldPartitionHelpers::DoCollectGarbage();
 					}
 
@@ -290,16 +288,25 @@ bool UWorldPartitionBuilder::Run(UWorld* World, FPackageSourceControlHelper& Pac
 	}
 	else
 	{
-		FBox BoundsToLoad(ForceInit);
+		TUniquePtr<FLoaderAdapterShape> LoaderAdapterShape;
+		
 		if (LoadingMode == ELoadingMode::EntireWorld)
 		{
-			BoundsToLoad += FBox(FVector(-WORLDPARTITION_MAX, -WORLDPARTITION_MAX, -WORLDPARTITION_MAX), FVector(WORLDPARTITION_MAX, WORLDPARTITION_MAX, WORLDPARTITION_MAX));
-			WorldPartition->LoadEditorCells(BoundsToLoad, false);
+			CellInfo.Bounds = FBox(FVector(-WORLDPARTITION_MAX, -WORLDPARTITION_MAX, -WORLDPARTITION_MAX), FVector(WORLDPARTITION_MAX, WORLDPARTITION_MAX, WORLDPARTITION_MAX));
+			LoaderAdapterShape = MakeUnique<FLoaderAdapterShape>(World, CellInfo.Bounds, TEXT("Loaded Region"));
+			LoaderAdapterShape->Load();
+		}
+		else
+		{
+			CellInfo.Bounds.Init();
 		}
-
-		CellInfo.Bounds = BoundsToLoad;
 
 		bResult = RunInternal(World, CellInfo, PackageHelper);
+
+		if (LoaderAdapterShape)
+		{
+			LoaderAdapterShape.Reset();
+		}
 	}
 
 	return PostRun(World, PackageHelper, bResult);
diff --git a/Engine/Source/Editor/WorldPartitionEditor/Private/WorldPartition/SWorldPartitionEditorGrid2D.cpp b/Engine/Source/Editor/WorldPartitionEditor/Private/WorldPartition/SWorldPartitionEditorGrid2D.cpp
--- a/Engine/Source/Editor/WorldPartitionEditor/Private/WorldPartition/SWorldPartitionEditorGrid2D.cpp
+++ b/Engine/Source/Editor/WorldPartitionEditor/Private/WorldPartition/SWorldPartitionEditorGrid2D.cpp
@@ -7,36 +7,87 @@
 #include "Modules/ModuleManager.h"
 #include "WorldPartition/WorldPartition.h"
 #include "WorldPartition/WorldPartitionEditorHash.h"
-#include "WorldPartition/WorldPartitionEditorCell.h"
 #include "WorldPartition/WorldPartitionActorDesc.h"
 #include "WorldPartition/WorldPartitionActorDescView.h"
+#include "WorldPartition/WorldPartitionActorLoaderInterface.h"
 #include "WorldPartition/WorldPartitionEditorPerProjectUserSettings.h"
+#include "WorldPartition/LoaderAdapter/LoaderAdapterShape.h"
 #include "Editor/GroupActor.h"
 #include "Editor/EditorEngine.h"
+#include "Fonts/FontMeasure.h"
 #include "Widgets/Input/SCheckBox.h"
 #include "Widgets/Input/SButton.h"
+#include "Rendering/SlateRenderer.h"
 #include "WorldBrowserModule.h"
 #include "LevelEditorViewport.h"
 #include "Algo/Transform.h"
 
 #define LOCTEXT_NAMESPACE "WorldPartitionEditor"
 
+static bool IsBoundsSelected(const FBox& SelectBox, const FBox& Bounds)
+{
+	return SelectBox.IsValid && Bounds.IntersectXY(SelectBox) && !Bounds.IsInsideXY(SelectBox);
+}
+
+template <class T>
+void ForEachSelectedLoaderAdapters(const UWorldPartition* WorldPartition, const FBox& SelectBox, T Func)
+{
+	for (IWorldPartitionActorLoaderInterface::ILoaderAdapter* LoaderAdapter : WorldPartition->GetRegisteredEditorLoaderAdapters())
+	{
+		if (LoaderAdapter->GetBoundingBox().IsSet())
+		{
+			if (IsBoundsSelected(SelectBox, *LoaderAdapter->GetBoundingBox()))
+			{
+				if (!Func(LoaderAdapter))
+				{
+					return;
+				}
+			}
+		}
+	}
+
+	WorldPartition->EditorHash->ForEachIntersectingActor(SelectBox, [&](FWorldPartitionActorDesc* ActorDesc)
+	{
+		if (AActor* Actor = ActorDesc->GetActor())
+		{
+			if (ActorDesc->GetActorNativeClass()->ImplementsInterface(UWorldPartitionActorLoaderInterface::StaticClass()))
+			{
+				if (IWorldPartitionActorLoaderInterface::ILoaderAdapter* LoaderAdapter = Cast<IWorldPartitionActorLoaderInterface>(Actor)->GetLoaderAdapter())
+				{
+					if (IsBoundsSelected(SelectBox, ActorDesc->GetBounds()))
+					{
+						if (!Func(LoaderAdapter))
+						{
+							return;
+						}
+					}
+				}
+			}
+		}
+	});
+}
+
 class FWorldPartitionActorDescViewBoundsProxy : public FWorldPartitionActorDescView
 {
 public:
 	FWorldPartitionActorDescViewBoundsProxy(const FWorldPartitionActorDesc* InActorDesc)
-	: FWorldPartitionActorDescView(InActorDesc)
+		: FWorldPartitionActorDescView(InActorDesc)
 	{}
 
 	FBox GetBounds() const
 	{
-		if (AActor* Actor = ActorDesc->GetActor())
+		if (AActor* Actor = GetActor())
 		{
 			return Actor->GetStreamingBounds();
 		}
 
 		return ActorDesc->GetBounds();
 	}
+
+	AActor* GetActor() const
+	{
+		return ActorDesc->GetActor(false);
+	}
 };
 
 SWorldPartitionEditorGrid2D::FEditorCommands::FEditorCommands()
@@ -51,9 +102,9 @@ SWorldPartitionEditorGrid2D::FEditorCommands::FEditorCommands()
 
 void SWorldPartitionEditorGrid2D::FEditorCommands::RegisterCommands()
 {
-	UI_COMMAND(LoadSelectedCells, "Load Selected Cells", "Load the selected cells.", EUserInterfaceActionType::Button, FInputChord());
-	UI_COMMAND(UnloadSelectedCells, "Unload Selected Cells", "Unload the selected cells.", EUserInterfaceActionType::Button, FInputChord());
-	UI_COMMAND(UnloadAllCells, "Unload All Cells", "Unload all cells.", EUserInterfaceActionType::Button, FInputChord());
+	UI_COMMAND(CreateRegionFromSelection, "Create Loading Region From Selection", "Create a loading region from the selection.", EUserInterfaceActionType::Button, FInputChord());
+	UI_COMMAND(LoadSelectedRegions, "Load Selected Regions", "Load the regions touched by the selection.", EUserInterfaceActionType::Button, FInputChord());
+	UI_COMMAND(UnloadSelectedRegions, "Unload Selected Regions", "Unload the regions touched by the selection.", EUserInterfaceActionType::Button, FInputChord());
 	UI_COMMAND(MoveCameraHere, "Move Camera Here", "Move the camera to the selected position.", EUserInterfaceActionType::Button, FInputChord());
 }
 
@@ -75,8 +126,7 @@ SWorldPartitionEditorGrid2D::SWorldPartitionEditorGrid2D()
 }
 
 SWorldPartitionEditorGrid2D::~SWorldPartitionEditorGrid2D()
-{
-}
+{}
 
 void SWorldPartitionEditorGrid2D::Construct(const FArguments& InArgs)
 {
@@ -126,9 +176,9 @@ void SWorldPartitionEditorGrid2D::Construct(const FArguments& InArgs)
 					.AutoWidth()
 					[
 						SNew(SCheckBox)
-						.IsChecked(GetMutableDefault<UWorldPartitionEditorPerProjectUserSettings>()->GetBugItGoLoadCells() ? ECheckBoxState::Checked : ECheckBoxState::Unchecked)
+						.IsChecked(GetMutableDefault<UWorldPartitionEditorPerProjectUserSettings>()->GetBugItGoLoadRegion() ? ECheckBoxState::Checked : ECheckBoxState::Unchecked)
 						.IsEnabled(true)
-						.OnCheckStateChanged(FOnCheckStateChanged::CreateLambda([=](ECheckBoxState State) { GetMutableDefault<UWorldPartitionEditorPerProjectUserSettings>()->SetBugItGoLoadCells(State == ECheckBoxState::Checked); }))
+						.OnCheckStateChanged(FOnCheckStateChanged::CreateLambda([=](ECheckBoxState State) { GetMutableDefault<UWorldPartitionEditorPerProjectUserSettings>()->SetBugItGoLoadRegion(State == ECheckBoxState::Checked); }))
 					]
 					+SHorizontalBox::Slot()
 					.FillWidth(1.0f)
@@ -137,7 +187,7 @@ void SWorldPartitionEditorGrid2D::Construct(const FArguments& InArgs)
 						SNew(STextBlock)
 						.AutoWrapText(true)
 						.IsEnabled(true)
-						.Text(LOCTEXT("BugItGoLoadCells", "BugItGo Load Cells"))
+						.Text(LOCTEXT("BugItGoLoadRegion", "BugItGo Load Region"))
 					]
 					+SHorizontalBox::Slot()
 					.AutoWidth()
@@ -174,35 +224,87 @@ void SWorldPartitionEditorGrid2D::Construct(const FArguments& InArgs)
 	const FEditorCommands& Commands = FEditorCommands::Get();
 	FUICommandList& ActionList = *CommandList;
 
-	auto CanLoadOrUnloadCells = [this]()
+	auto CanCreateRegionFromSelection = [this]()
+	{
+		return !!SelectBox.IsValid;
+	};
+
+	auto CanLoadUnloadSelectedRegions = [this](bool bLoad)
 	{
-		return WorldPartition && SelectBox.IsValid && (WorldPartition->EditorHash->ForEachIntersectingCell(SelectBox, [&](UWorldPartitionEditorCell*){}) > 0);
+		bool bResult = false;
+		ForEachSelectedLoaderAdapters(WorldPartition, SelectBox, [&bResult, bLoad](IWorldPartitionActorLoaderInterface::ILoaderAdapter* LoaderAdapter)
+		{
+			if (bLoad != LoaderAdapter->IsLoaded())
+			{
+				bResult = true;
+				return false;
+			}
+			return true;
+		});
+		return bResult;
 	};
 
-	ActionList.MapAction(Commands.LoadSelectedCells, FExecuteAction::CreateSP(this, &SWorldPartitionEditorGrid2D::LoadSelectedCells), FCanExecuteAction::CreateLambda(CanLoadOrUnloadCells));
-	ActionList.MapAction(Commands.UnloadSelectedCells, FExecuteAction::CreateSP(this, &SWorldPartitionEditorGrid2D::UnloadSelectedCells), FCanExecuteAction::CreateLambda(CanLoadOrUnloadCells));
-	ActionList.MapAction(Commands.UnloadAllCells, FExecuteAction::CreateSP(this, &SWorldPartitionEditorGrid2D::UnloadAllCells));
+	auto CanLoadSelectedRegions = [this, CanLoadUnloadSelectedRegions]() { return CanLoadUnloadSelectedRegions(true); };
+	auto CanUnloadSelectedRegions = [this, CanLoadUnloadSelectedRegions]() { return CanLoadUnloadSelectedRegions(false); };
+
+	ActionList.MapAction(Commands.CreateRegionFromSelection, FExecuteAction::CreateSP(this, &SWorldPartitionEditorGrid2D::CreateRegionFromSelection), FCanExecuteAction::CreateLambda(CanCreateRegionFromSelection));
+	ActionList.MapAction(Commands.LoadSelectedRegions, FExecuteAction::CreateSP(this, &SWorldPartitionEditorGrid2D::LoadSelectedRegions), FCanExecuteAction::CreateLambda(CanLoadSelectedRegions));
+	ActionList.MapAction(Commands.UnloadSelectedRegions, FExecuteAction::CreateSP(this, &SWorldPartitionEditorGrid2D::UnloadSelectedRegions), FCanExecuteAction::CreateLambda(CanUnloadSelectedRegions));
 	ActionList.MapAction(Commands.MoveCameraHere, FExecuteAction::CreateSP(this, &SWorldPartitionEditorGrid2D::MoveCameraHere));
 }
 
-void SWorldPartitionEditorGrid2D::LoadSelectedCells()
+void SWorldPartitionEditorGrid2D::CreateRegionFromSelection()
 {
-	WorldPartition->LoadEditorCells(SelectBox, true);
+	const FBox RegionBox(FVector(SelectBox.Min.X, SelectBox.Min.Y, -WORLDPARTITION_MAX), FVector(SelectBox.Max.X, SelectBox.Max.Y, WORLDPARTITION_MAX));
+	FLoaderAdapterShape* LoaderAdapter = WorldPartition->CreateEditorLoaderAdapter<FLoaderAdapterShape>(World, RegionBox, TEXT("Loaded Region"));
+	LoaderAdapter->SetUserCreated(true);
+	LoaderAdapter->Load();
+
+	SelectBox.Init();
+
 	GEditor->RedrawLevelEditingViewports();
 	Refresh();
 }
 
-void SWorldPartitionEditorGrid2D::UnloadSelectedCells()
+void SWorldPartitionEditorGrid2D::LoadSelectedRegions()
 {
-	WorldPartition->UnloadEditorCells(SelectBox, true);
+	check(SelectBox.IsValid);
+
+	ForEachSelectedLoaderAdapters(WorldPartition, SelectBox, [](IWorldPartitionActorLoaderInterface::ILoaderAdapter* LoaderAdapter)
+	{
+		LoaderAdapter->Load();
+		return true;
+	});
+
+	SelectBox.Init();
+
 	GEditor->RedrawLevelEditingViewports();
 	Refresh();
 }
 
-void SWorldPartitionEditorGrid2D::UnloadAllCells()
+void SWorldPartitionEditorGrid2D::UnloadSelectedRegions()
 {
-	const FBox AllCellsBox(FVector(-WORLDPARTITION_MAX, -WORLDPARTITION_MAX, -WORLDPARTITION_MAX), FVector(WORLDPARTITION_MAX, WORLDPARTITION_MAX, WORLDPARTITION_MAX));
-	WorldPartition->UnloadEditorCells(AllCellsBox, true);
+	check(SelectBox.IsValid);
+
+	TArray<IWorldPartitionActorLoaderInterface::ILoaderAdapter*> LoaderAdaptersToRelease;
+	const TSet<IWorldPartitionActorLoaderInterface::ILoaderAdapter*>& RegisteredLoaderAdapters = WorldPartition->GetRegisteredEditorLoaderAdapters();
+
+	ForEachSelectedLoaderAdapters(WorldPartition, SelectBox, [&LoaderAdaptersToRelease, &RegisteredLoaderAdapters](IWorldPartitionActorLoaderInterface::ILoaderAdapter* LoaderAdapter)
+	{
+		if (LoaderAdapter->Unload() && RegisteredLoaderAdapters.Contains(LoaderAdapter))
+		{
+			LoaderAdaptersToRelease.Add(LoaderAdapter);
+		}
+		return true;
+	});
+
+	for (IWorldPartitionActorLoaderInterface::ILoaderAdapter* LoaderAdapterToRelease : LoaderAdaptersToRelease)
+	{
+		WorldPartition->ReleaseEditorLoaderAdapter(LoaderAdapterToRelease);
+	}
+
+	SelectBox.Init();
+
 	GEditor->RedrawLevelEditingViewports();
 	Refresh();
 }
@@ -237,7 +339,7 @@ FReply SWorldPartitionEditorGrid2D::OnMouseButtonDown(const FGeometry& MyGeometr
 			bIsSelecting = true;
 			SelectionStart = MouseCursorPosWorld;
 			SelectionEnd = SelectionStart;
-			UpdateSelection();
+			SelectBox.Init();
 		}
 
 		return ReplyState;
@@ -265,15 +367,17 @@ FReply SWorldPartitionEditorGrid2D::OnMouseButtonUp(const FGeometry& MyGeometry,
 
 			const FEditorCommands& Commands = FEditorCommands::Get();
 
-			MenuBuilder.BeginSection(NAME_None, LOCTEXT("WorldPartition", "Selected Actors"));
-			{
-				MenuBuilder.AddMenuEntry(Commands.LoadSelectedCells);
-				MenuBuilder.AddMenuEntry(Commands.UnloadSelectedCells);
-				MenuBuilder.AddMenuEntry(Commands.UnloadAllCells);
-				MenuBuilder.AddMenuEntry(Commands.MoveCameraHere);
-			}
+			MenuBuilder.BeginSection(NAME_None, LOCTEXT("WorldPartitionSelection", "Selection"));
+				MenuBuilder.AddMenuEntry(Commands.CreateRegionFromSelection);
+				MenuBuilder.AddMenuSeparator();
+				MenuBuilder.AddMenuEntry(Commands.LoadSelectedRegions);
+				MenuBuilder.AddMenuEntry(Commands.UnloadSelectedRegions);
 			MenuBuilder.EndSection();
 
+			MenuBuilder.BeginSection(NAME_None, LOCTEXT("WorldPartitionMisc", "Misc"));
+				MenuBuilder.AddMenuEntry(Commands.MoveCameraHere);
+			MenuBuilder.EndSection();			
+
 			FWidgetPath WidgetPath = MouseEvent.GetEventPath() != nullptr ? *MouseEvent.GetEventPath() : FWidgetPath();
 			FSlateApplication::Get().PushMenu(AsShared(), WidgetPath, MenuBuilder.MakeWidget(), MouseEvent.GetScreenSpacePosition(), FPopupTransitionEffect(FPopupTransitionEffect::ContextMenu));
 		}
@@ -406,18 +510,35 @@ uint32 SWorldPartitionEditorGrid2D::PaintActors(const FGeometry& AllottedGeometr
 	const FBox ViewRectWorld(FVector(WorldViewRect.Min.X, WorldViewRect.Min.Y, -WORLDPARTITION_MAX), FVector(WorldViewRect.Max.X, WorldViewRect.Max.Y, WORLDPARTITION_MAX));
 
 	TSet<FWorldPartitionActorDescViewBoundsProxy> ActorDescList;
+	
+	TArray<IWorldPartitionActorLoaderInterface::ILoaderAdapter*> AllLoaderAdapters;
+	for (IWorldPartitionActorLoaderInterface::ILoaderAdapter* LoaderAdapter: WorldPartition->GetRegisteredEditorLoaderAdapters())
+	{
+		AllLoaderAdapters.Add(LoaderAdapter);
+	}
 
-	// Include all actors if requested
-	if (bShowActors)
+	WorldPartition->EditorHash->ForEachIntersectingActor(ViewRectWorld, [&](FWorldPartitionActorDesc* ActorDesc)
 	{
-		WorldPartition->EditorHash->ForEachIntersectingActor(ViewRectWorld, [&](FWorldPartitionActorDesc* ActorDesc)
+		FWorldPartitionActorDescViewBoundsProxy ActorDescViewProxy(ActorDesc);
+
+		if (bShowActors)
 		{
-			ActorDescList.Emplace(ActorDesc);
-		});
-	}
+			if (ActorDescViewProxy.GetIsSpatiallyLoaded())
+			{
+				ActorDescList.Emplace(ActorDescViewProxy);
+			}
+		}
+		
+		if (ActorDesc->GetActorNativeClass()->ImplementsInterface(UWorldPartitionActorLoaderInterface::StaticClass()))
+		{
+			if (AActor* Actor = ActorDescViewProxy.GetActor())
+			{				
+				AllLoaderAdapters.Add(Cast<IWorldPartitionActorLoaderInterface>(Actor)->GetLoaderAdapter());
+			}
+		}
+	});
 
 	// Include selected actors
-	FWorldPartitionActorDesc::GlobalTag++;
 	for (FSelectionIterator It = GEditor->GetSelectedActorIterator(); It; ++It)
 	{
 		if (AActor* Actor = Cast<AActor>(*It))
@@ -425,15 +546,126 @@ uint32 SWorldPartitionEditorGrid2D::PaintActors(const FGeometry& AllottedGeometr
 			if (FWorldPartitionActorDesc* ActorDesc = WorldPartition->GetActorDesc(Actor->GetActorGuid()))
 			{
 				ActorDescList.Emplace(ActorDesc);
-				ActorDesc->Tag = FWorldPartitionActorDesc::GlobalTag;
 			}
 		}
 	}
 
-	if (ActorDescList.Num())
+	if (AllLoaderAdapters.Num())
 	{
 		const FLinearColor LoadedActorColor(0.75f, 0.75f, 0.75f, 1.0f);
 		const FLinearColor UnloadedActorColor(0.5f, 0.5f, 0.5f, 1.0f);	
+
+		TArray<FVector2D> LinePoints;
+		LinePoints.SetNum(5);
+
+		for (const IWorldPartitionActorLoaderInterface::ILoaderAdapter* LoaderAdapter : AllLoaderAdapters)
+		{
+			if (LoaderAdapter->GetBoundingBox().IsSet())
+			{
+				const FBox ActorBounds = *LoaderAdapter->GetBoundingBox();
+
+				FVector Origin, Extent;
+				ActorBounds.GetCenterAndExtents(Origin, Extent);
+
+				const FVector2D TopLeftW = FVector2D(Origin - Extent);
+				const FVector2D BottomRightW = FVector2D(Origin + Extent);
+				const FVector2D TopRightW = FVector2D(BottomRightW.X, TopLeftW.Y);
+				const FVector2D BottomLeftW = FVector2D(TopLeftW.X, BottomRightW.Y);
+
+				const FVector2D TopLeft = WorldToScreen.TransformPoint(TopLeftW);
+				const FVector2D BottomRight = WorldToScreen.TransformPoint(BottomRightW);
+				const FVector2D TopRight = WorldToScreen.TransformPoint(TopRightW);
+				const FVector2D BottomLeft = WorldToScreen.TransformPoint(BottomLeftW);
+
+				const FBox2D ActorViewBox(TopLeft, BottomRight);
+
+				const float FullScreenColorGradient = FMath::Min(ViewRect.GetArea() / ActorViewBox.GetArea(), 1.0f);
+
+				if (FullScreenColorGradient > 0.0f)
+				{
+					const float MinimumAreaCull = 32.0f;
+					const float AreaFadeDistance = 128.0f;
+					if ((Extent.Size2D() < KINDA_SMALL_NUMBER) || (ActorViewBox.GetArea() > MinimumAreaCull))
+					{
+						const FPaintGeometry ActorGeometry = AllottedGeometry.ToPaintGeometry(TopLeft, BottomRight - TopLeft);
+						const float LoaderColorGradient = FMath::Min((ActorViewBox.GetArea() - MinimumAreaCull) / AreaFadeDistance, 1.0f);
+
+						// Highlight
+						{
+							const FSlateColorBrush LoadedBrush(FLinearColor::White);
+							const FSlateColorBrush UnloadedBrush(FLinearColor::Black);
+							const FLinearColor LoadedColor(1.0f, 1.0f, 1.0f, 0.25f * LoaderColorGradient * FullScreenColorGradient);
+							const FLinearColor UnloadedColor(0.0f, 0.0f, 0.0f, 0.25f * LoaderColorGradient * FullScreenColorGradient);
+
+							FSlateDrawElement::MakeBox(
+								OutDrawElements,
+								++LayerId,
+								ActorGeometry,
+								LoaderAdapter->IsLoaded() ? &LoadedBrush : &UnloadedBrush,
+								ESlateDrawEffect::None,
+								LoaderAdapter->IsLoaded() ? LoadedColor : UnloadedColor
+							);
+						}
+
+						// Outline
+						{
+							const bool IsHighlighted = IsBoundsSelected(SelectBox, ActorBounds);
+							const FLinearColor OutlineColor(IsHighlighted ? 1.0f : 0.0f, IsHighlighted ? 1.0f : 0.0f, 0.0f, LoaderColorGradient * FullScreenColorGradient);
+
+							LinePoints[0] = TopLeft;
+							LinePoints[1] = TopRight;
+							LinePoints[2] = BottomRight;
+							LinePoints[3] = BottomLeft;
+							LinePoints[4] = TopLeft;
+
+							FSlateDrawElement::MakeLines
+							(
+								OutDrawElements,
+								++LayerId,
+								AllottedGeometry.ToPaintGeometry(),
+								LinePoints,
+								ESlateDrawEffect::None,
+								OutlineColor,
+								true,
+								2.0f
+							);
+						}
+
+						// Label
+						{
+							const FString ActorLabel = *LoaderAdapter->GetLabel();
+							const FVector2D ActorViewBoxCenter = ActorViewBox.GetCenter();
+							const FVector2D LabelTextSize = FSlateApplication::Get().GetRenderer()->GetFontMeasureService()->Measure(ActorLabel, SmallLayoutFont);
+
+							if (LabelTextSize.X > 0)
+							{
+								const FVector2D LabelTextPos = ActorViewBoxCenter - LabelTextSize * 0.5f;					
+								const float LabelColorGradient = FMath::Clamp(ActorGeometry.GetLocalSize().X / LabelTextSize.X - 1.0f, 0.0f, 1.0f);					
+
+								if (LabelColorGradient > 0.0f)
+								{
+									const FLinearColor LabelColor(1.0f, 1.0f, 1.0f, LabelColorGradient * LoaderColorGradient * FullScreenColorGradient);
+
+									FSlateDrawElement::MakeText(
+										OutDrawElements,
+										++LayerId,
+										AllottedGeometry.ToPaintGeometry(LabelTextPos, FVector2D(1,1)),
+										ActorLabel,
+										SmallLayoutFont,
+										ESlateDrawEffect::None,
+										LabelColor
+									);
+								}
+							}
+						}
+					}
+				}
+			}
+		}
+	}
+
+	if (ActorDescList.Num())
+	{
 		const FLinearColor SelectedActorColor(1.0f, 1.0f, 1.0f, 1.0f);
 
 		TArray<FVector2D> LinePoints;
@@ -459,17 +691,20 @@ uint32 SWorldPartitionEditorGrid2D::PaintActors(const FGeometry& AllottedGeometr
 
 			const float MinimumAreaCull = 32.0f;
 			const float AreaFadeDistance = 128.0f;
-			if (ActorViewBox.Intersect(ViewRect) && ((Extent.Size2D() < KINDA_SMALL_NUMBER) || ActorViewBox.GetArea() > MinimumAreaCull))
+			if ((Extent.Size2D() < KINDA_SMALL_NUMBER) || (ActorViewBox.GetArea() > MinimumAreaCull))
 			{
-				FPaintGeometry ActorGeometry = AllottedGeometry.ToPaintGeometry(TopLeft, BottomRight - TopLeft);
-				float ActorColorGradient = FMath::Min((ActorViewBox.GetArea() - MinimumAreaCull) / AreaFadeDistance, 1.0f);
-				float ActorBrightness = ActorDescView.GetIsSpatiallyLoaded() ? 1.0f : 0.3f;
+				const FPaintGeometry ActorGeometry = AllottedGeometry.ToPaintGeometry(TopLeft, BottomRight - TopLeft);
+
+				const float ActorColorGradient = FMath::Min((ActorViewBox.GetArea() - MinimumAreaCull) / AreaFadeDistance, 1.0f);
+				const float ActorBrightness = ActorDescView.GetIsSpatiallyLoaded() ? 1.0f : 0.3f;
 				FLinearColor ActorColor(ActorBrightness, ActorBrightness, ActorBrightness, ActorColorGradient);
 
 				UClass* ActorClass = ActorDescView.GetActorNativeClass();
+				const AActor* Actor = ActorDescView.GetActor();
 
+				const bool bIsSelected = Actor ? Actor->IsSelected() : false;
 				const float SquaredDistanceToPoint = ActorViewBox.ComputeSquaredDistanceToPoint(MouseCursorPos);
-				if ((ActorDescView.GetTag() == FWorldPartitionActorDesc::GlobalTag) || (SquaredDistanceToPoint > 0.0f && SquaredDistanceToPoint <= 2.0f))
+				if (bIsSelected || (SquaredDistanceToPoint > 0.0f && SquaredDistanceToPoint <= 2.0f))
 				{
 					ActorColor = FLinearColor::Yellow;
 
@@ -619,13 +854,13 @@ uint32 SWorldPartitionEditorGrid2D::PaintViewer(const FGeometry& AllottedGeometr
 
 uint32 SWorldPartitionEditorGrid2D::PaintSelection(const FGeometry& AllottedGeometry, const FSlateRect& MyCullingRect, FSlateWindowElementList& OutDrawElements, uint32 LayerId) const
 {
-	if (bIsSelecting)
+	if (SelectBox.IsValid)
 	{
 		TArray<FVector2D> LinePoints;
 		LinePoints.SetNum(5);
 
-		FVector2D TopLeftW = SelectionStart;
-		FVector2D BottomRightW = SelectionEnd;
+		FVector2D TopLeftW = FVector2D(SelectBox.Min);
+		FVector2D BottomRightW = FVector2D(SelectBox.Max);
 		FVector2D TopRightW = FVector2D(BottomRightW.X, TopLeftW.Y);
 		FVector2D BottomLeftW = FVector2D(TopLeftW.X, BottomRightW.Y);
 
@@ -659,7 +894,16 @@ uint32 SWorldPartitionEditorGrid2D::PaintSelection(const FGeometry& AllottedGeom
 			);
 		}
 
-		FSlateDrawElement::MakeLines(OutDrawElements, LayerId, AllottedGeometry.ToPaintGeometry(), LinePoints, ESlateDrawEffect::None, FLinearColor::White, false, 1.0f);
+		FSlateDrawElement::MakeLines(
+			OutDrawElements, 
+			LayerId, 
+			AllottedGeometry.ToPaintGeometry(), 
+			LinePoints, 
+			ESlateDrawEffect::None, 
+			FLinearColor::White, 
+			true, 
+			2.0f
+		);
 	}
 
 	return LayerId + 1;
@@ -772,14 +1016,18 @@ void SWorldPartitionEditorGrid2D::UpdateTransform() const
 
 void SWorldPartitionEditorGrid2D::UpdateSelection()
 {
-	SelectBox.Init();
-
-	const FBox SelectionBox(
-		FVector(FMath::Min(SelectionStart.X, SelectionEnd.X), FMath::Min(SelectionStart.Y, SelectionEnd.Y), -WORLDPARTITION_MAX),
-		FVector(FMath::Max(SelectionStart.X, SelectionEnd.X), FMath::Max(SelectionStart.Y, SelectionEnd.Y), WORLDPARTITION_MAX)
+	const FBox2D SelectBox2D(
+		FVector2D::Min(SelectionStart, SelectionEnd), 
+		FVector2D::Max(SelectionStart, SelectionEnd)
 	);
 
-	SelectBox = SelectionBox;
+	if (SelectBox2D.GetArea() > 0.0f)
+	{
+		SelectBox = FBox(
+			FVector(SelectBox2D.Min.X, SelectBox2D.Min.Y, -WORLDPARTITION_MAX),
+			FVector(SelectBox2D.Max.X, SelectBox2D.Max.Y, WORLDPARTITION_MAX)
+		);
+	}
 }
 
 #undef LOCTEXT_NAMESPACE
diff --git a/Engine/Source/Editor/WorldPartitionEditor/Private/WorldPartition/SWorldPartitionEditorGridSpatialHash.cpp b/Engine/Source/Editor/WorldPartitionEditor/Private/WorldPartition/SWorldPartitionEditorGridSpatialHash.cpp
--- a/Engine/Source/Editor/WorldPartitionEditor/Private/WorldPartition/SWorldPartitionEditorGridSpatialHash.cpp
+++ b/Engine/Source/Editor/WorldPartitionEditor/Private/WorldPartition/SWorldPartitionEditorGridSpatialHash.cpp
@@ -3,7 +3,6 @@
 #include "GameFramework/WorldSettings.h"
 #include "DerivedDataCache/Public/DerivedDataCacheInterface.h"
 #include "Misc/SecureHash.h"
-#include "WorldPartition/WorldPartitionEditorCell.h"
 #include "WorldPartition/WorldPartitionRuntimeSpatialHash.h"
 #include "WorldPartition/WorldPartitionEditorSpatialHash.h"
 #include "WorldPartition/WorldPartitionActorDesc.h"
@@ -32,14 +31,17 @@
 
 #define LOCTEXT_NAMESPACE "WorldPartitionEditor"
 
-SWorldPartitionEditorGridSpatialHash::SWorldPartitionEditorGridSpatialHash()
-	: WorldMiniMapBounds(ForceInit)
+static inline FBox2D ToBox2D(const FBox& Box)
 {
+	return FBox2D(FVector2D(Box.Min), FVector2D(Box.Max));
 }
 
+SWorldPartitionEditorGridSpatialHash::SWorldPartitionEditorGridSpatialHash()
+	: WorldMiniMapBounds(ForceInit)
+{}
+
 SWorldPartitionEditorGridSpatialHash::~SWorldPartitionEditorGridSpatialHash()
-{
-}
+{}
 
 void SWorldPartitionEditorGridSpatialHash::Construct(const FArguments& InArgs)
 {
@@ -58,8 +60,7 @@ void SWorldPartitionEditorGridSpatialHash::Construct(const FArguments& InArgs)
 
 void SWorldPartitionEditorGridSpatialHash::UpdateWorldMiniMapDetails()
 {
-	AWorldPartitionMiniMap* WorldMiniMap = FWorldPartitionMiniMapHelper::GetWorldPartitionMiniMap(World);
-	if (WorldMiniMap)
+	if (AWorldPartitionMiniMap* WorldMiniMap = FWorldPartitionMiniMapHelper::GetWorldPartitionMiniMap(World))
 	{
 		WorldMiniMapBounds = FBox2D(FVector2D(WorldMiniMap->MiniMapWorldBounds.Min), FVector2D(WorldMiniMap->MiniMapWorldBounds.Max));
 		if (UTexture2D* MiniMapTexture = WorldMiniMap->MiniMapTexture)
@@ -73,21 +74,35 @@ void SWorldPartitionEditorGridSpatialHash::UpdateWorldMiniMapDetails()
 
 int32 SWorldPartitionEditorGridSpatialHash::PaintGrid(const FGeometry& AllottedGeometry, const FSlateRect& MyCullingRect, FSlateWindowElementList& OutDrawElements, int32 LayerId) const
 {
+	UWorldPartitionEditorSpatialHash* EditorSpatialHash = (UWorldPartitionEditorSpatialHash*)WorldPartition->EditorHash;
+	
+	// Found the best cell size depending on the current zoom
+	int32 EffectiveCellSize = EditorSpatialHash->CellSize;
+	for (const UWorldPartitionEditorSpatialHash::FCellNodeHashLevel& HashLevel : EditorSpatialHash->HashLevels)
+	{
+		const FVector2D CellScreenSize = WorldToScreen.TransformVector(FVector2D(EffectiveCellSize, EffectiveCellSize));
+		if (CellScreenSize.X > 32)
+		{
+			break;
+		}
+
+		EffectiveCellSize *= 2;
+	}
+	
+	// Compute visible rect
 	const FBox2D ViewRect(FVector2D(ForceInitToZero), AllottedGeometry.GetLocalSize());
 	const FBox2D ViewRectWorld(ScreenToWorld.TransformPoint(ViewRect.Min), ScreenToWorld.TransformPoint(ViewRect.Max));
 
-	UWorldPartitionEditorSpatialHash* EditorSpatialHash = (UWorldPartitionEditorSpatialHash*)WorldPartition->EditorHash;
-	
 	FBox VisibleGridRectWorld(
 		FVector(
-			FMath::Max(FMath::FloorToFloat(EditorSpatialHash->Bounds.Min.X / EditorSpatialHash->CellSize) * EditorSpatialHash->CellSize, ViewRectWorld.Min.X),
-			FMath::Max(FMath::FloorToFloat(EditorSpatialHash->Bounds.Min.Y / EditorSpatialHash->CellSize) * EditorSpatialHash->CellSize, ViewRectWorld.Min.Y),
-			FMath::FloorToFloat(EditorSpatialHash->Bounds.Min.Z / EditorSpatialHash->CellSize) * EditorSpatialHash->CellSize
+			FMath::Max(FMath::FloorToFloat(EditorSpatialHash->Bounds.Min.X / EffectiveCellSize) * EffectiveCellSize, ViewRectWorld.Min.X),
+			FMath::Max(FMath::FloorToFloat(EditorSpatialHash->Bounds.Min.Y / EffectiveCellSize) * EffectiveCellSize, ViewRectWorld.Min.Y),
+			FMath::FloorToFloat(EditorSpatialHash->Bounds.Min.Z / EffectiveCellSize) * EffectiveCellSize
 		),
 		FVector(
-			FMath::Min(FMath::FloorToFloat(EditorSpatialHash->Bounds.Max.X / EditorSpatialHash->CellSize) * EditorSpatialHash->CellSize, ViewRectWorld.Max.X),
-			FMath::Min(FMath::FloorToFloat(EditorSpatialHash->Bounds.Max.Y / EditorSpatialHash->CellSize) * EditorSpatialHash->CellSize, ViewRectWorld.Max.Y),
-			FMath::FloorToFloat(EditorSpatialHash->Bounds.Max.Z / EditorSpatialHash->CellSize) * EditorSpatialHash->CellSize
+			FMath::Min(FMath::CeilToFloat(EditorSpatialHash->Bounds.Max.X / EffectiveCellSize) * EffectiveCellSize, ViewRectWorld.Max.X),
+			FMath::Min(FMath::CeilToFloat(EditorSpatialHash->Bounds.Max.Y / EffectiveCellSize) * EffectiveCellSize, ViewRectWorld.Max.Y),
+			FMath::CeilToFloat(EditorSpatialHash->Bounds.Max.Z / EffectiveCellSize) * EffectiveCellSize
 		)
 	);
 
@@ -153,166 +168,25 @@ int32 SWorldPartitionEditorGridSpatialHash::PaintGrid(const FGeometry& AllottedG
 		}
 	}
 
-	struct FCellDesc2D
-	{
-		FCellDesc2D()
-			: Bounds(ForceInitToZero)
-			, NumLoaded(0)
-			, NumUnloaded(0)
-			, NumEmpty(0)
-		{}
-
-		FBox2D Bounds;
-		int32 NumLoaded;
-		int32 NumUnloaded;
-		int32 NumEmpty;
-	};
-
-	// Draw shadowed regions
-	{
-		TMap<uint32, FCellDesc2D> FlattenedCells2D;
-		EditorSpatialHash->ForEachIntersectingCells(VisibleGridRectWorld.ExpandBy(-1), 0, [&](const UWorldPartitionEditorSpatialHash::FCellCoord& CellCoord)
-		{
-			FHashBuilder HashBuilder;
-			HashBuilder << CellCoord.X << CellCoord.Y;
-			uint32 CellHash2D = HashBuilder.GetHash();			
-			
-			FCellDesc2D& CellDesc2D = FlattenedCells2D.FindOrAdd(CellHash2D);
-			
-			FBox CellBounds = EditorSpatialHash->GetCellBounds(CellCoord);
-			CellDesc2D.Bounds = FBox2D(FVector2D(CellBounds.Min), FVector2D(CellBounds.Max));
-
-			if (UWorldPartitionEditorCell** CellPtr = EditorSpatialHash->HashCells.Find(CellCoord))
-			{
-				if ((*CellPtr)->IsEmpty())
-				{
-					CellDesc2D.NumEmpty++;
-				}
-				else if ((*CellPtr)->IsLoaded())
-				{
-					CellDesc2D.NumLoaded++;
-				}
-				else
-				{
-					CellDesc2D.NumUnloaded++;
-				}
-			}
-			else
-			{
-				CellDesc2D.NumEmpty++;
-			}
-		});
-
-		for(auto& UniqueCell: FlattenedCells2D)
-		{
-			const FCellDesc2D& Cell = UniqueCell.Value;
-
-			// Fully loaded
-			if (Cell.NumLoaded && !Cell.NumUnloaded)
-			{
-				continue;
-			}
-
-			const FPaintGeometry CellGeometry = AllottedGeometry.ToPaintGeometry(
-				WorldToScreen.TransformPoint(Cell.Bounds.Min),
-				WorldToScreen.TransformPoint(Cell.Bounds.Max) - WorldToScreen.TransformPoint(Cell.Bounds.Min)
-			);
-
-			const FSlateColorBrush CellBrush(FLinearColor::White);
-			const FLinearColor FullyUnloadedCellColor(0, 0, 0, 0.5f);
-			const FLinearColor PartiallyLoadedCellColor(0, 0, 0, 0.25f);
-			const FLinearColor EmptyCellColor(0, 0, 0, 0.75f);
-			
-			FLinearColor CellColor;
-
-			// Partially loaded aread
-			if (Cell.NumLoaded && Cell.NumUnloaded)
-			{
-				CellColor = PartiallyLoadedCellColor;
-			}
-			else if (Cell.NumEmpty && !Cell.NumUnloaded && !Cell.NumLoaded)
-			{
-				CellColor = EmptyCellColor;
-			}
-			else
-			{
-				CellColor = FullyUnloadedCellColor;
-			}
-
-			FSlateDrawElement::MakeBox(
-				OutDrawElements,
-				++LayerId,
-				CellGeometry,
-				&CellBrush,
-				ESlateDrawEffect::None,
-				CellColor
-			);
-		}
-	}
-
-	// Draw selected cells
-	if (SelectBox.IsValid)
-	{
-		const FBox VisibleSelectBox = SelectBox.Overlap(VisibleGridRectWorld);
-
-		if (VisibleSelectBox.IsValid)
-		{
-			TMap<uint32, FCellDesc2D> FlattenedCells2D;
-			WorldPartition->EditorHash->ForEachIntersectingCell(VisibleSelectBox, [&](UWorldPartitionEditorCell* Cell)
-			{
-				UWorldPartitionEditorSpatialHash::FCellCoord CellCoord = EditorSpatialHash->GetCellCoords(Cell->Bounds.GetCenter(), 0);
-
-				FHashBuilder HashBuilder;
-				HashBuilder << CellCoord.X << CellCoord.Y;
-				uint32 CellHash2D = HashBuilder.GetHash();
-				FCellDesc2D& CellDesc2D = FlattenedCells2D.Add(CellHash2D);
-
-				CellDesc2D.Bounds = FBox2D(FVector2D(Cell->Bounds.Min), FVector2D(Cell->Bounds.Max));
-			});
-
-			for (auto& UniqueCell : FlattenedCells2D)
-			{
-				const FCellDesc2D& Cell = UniqueCell.Value;
-
-				FPaintGeometry CellGeometry = AllottedGeometry.ToPaintGeometry(
-					WorldToScreen.TransformPoint(Cell.Bounds.Min),
-					WorldToScreen.TransformPoint(Cell.Bounds.Max) - WorldToScreen.TransformPoint(Cell.Bounds.Min)
-				);
-
-				FSlateColorBrush CellBrush(FLinearColor::White);
-				FLinearColor CellColor(1, 1, 1, 0.25f);
-
-				FSlateDrawElement::MakeBox(
-					OutDrawElements,
-					++LayerId,
-					CellGeometry,
-					&CellBrush,
-					ESlateDrawEffect::None,
-					CellColor
-				);
-			}
-		}
-	}
-
-	if (FBox2D(FVector2D(VisibleGridRectWorld.Min), FVector2D(VisibleGridRectWorld.Max)).GetArea() > 0.0f)
+	if (ToBox2D(VisibleGridRectWorld).GetArea() > 0.0f)
 	{
 		const FLinearColor Color = FLinearColor(0.1f, 0.1f, 0.1f, 1.f);
 
 		FIntVector2 TopLeftW(
-			FMath::FloorToFloat(VisibleGridRectWorld.Min.X / EditorSpatialHash->CellSize) * EditorSpatialHash->CellSize,
-			FMath::FloorToFloat(VisibleGridRectWorld.Min.Y / EditorSpatialHash->CellSize) * EditorSpatialHash->CellSize
+			FMath::FloorToFloat(VisibleGridRectWorld.Min.X / EffectiveCellSize) * EffectiveCellSize,
+			FMath::FloorToFloat(VisibleGridRectWorld.Min.Y / EffectiveCellSize) * EffectiveCellSize
 		);
 
 		FIntVector2 BottomRightW(
-			FMath::CeilToFloat(VisibleGridRectWorld.Max.X / EditorSpatialHash->CellSize) * EditorSpatialHash->CellSize,
-			FMath::CeilToFloat(VisibleGridRectWorld.Max.Y / EditorSpatialHash->CellSize) * EditorSpatialHash->CellSize
+			FMath::CeilToFloat(VisibleGridRectWorld.Max.X / EffectiveCellSize) * EffectiveCellSize,
+			FMath::CeilToFloat(VisibleGridRectWorld.Max.Y / EffectiveCellSize) * EffectiveCellSize
 		);
 
 		TArray<FVector2D> LinePoints;
 		LinePoints.SetNum(2);
 
 		// Horizontal
-		for (int32 i=TopLeftW.Y; i<=BottomRightW.Y; i+=EditorSpatialHash->CellSize)
+		for (int32 i=TopLeftW.Y; i<=BottomRightW.Y; i+=EffectiveCellSize)
 		{
 			FVector2D LineStartH(TopLeftW.X, i);
 			FVector2D LineEndH(BottomRightW.X, i);
@@ -324,7 +198,7 @@ int32 SWorldPartitionEditorGridSpatialHash::PaintGrid(const FGeometry& AllottedG
 		}
 
 		// Vertical
-		for (int32 i=TopLeftW.X; i<=BottomRightW.X; i+=EditorSpatialHash->CellSize)
+		for (int32 i=TopLeftW.X; i<=BottomRightW.X; i+=EffectiveCellSize)
 		{
 			FVector2D LineStartH(i, TopLeftW.Y);
 			FVector2D LineEndH(i, BottomRightW.Y);
@@ -336,10 +210,10 @@ int32 SWorldPartitionEditorGridSpatialHash::PaintGrid(const FGeometry& AllottedG
 		}
 
 		// Draw coordinates
-		if (GetMutableDefault<UWorldPartitionEditorPerProjectUserSettings>()->GetShowCellCoords())
+		if ((EffectiveCellSize == EditorSpatialHash->CellSize) && GetMutableDefault<UWorldPartitionEditorPerProjectUserSettings>()->GetShowCellCoords())
 		{
 			const TSharedRef<FSlateFontMeasure> FontMeasure = FSlateApplication::Get().GetRenderer()->GetFontMeasureService();
-			const FVector2D CellScreenSize = WorldToScreen.TransformVector(FVector2D(EditorSpatialHash->CellSize, EditorSpatialHash->CellSize));
+			const FVector2D CellScreenSize = WorldToScreen.TransformVector(FVector2D(EffectiveCellSize, EffectiveCellSize));
 
 			FSlateFontInfo CoordsFont;
 			FVector2D DefaultCoordTextSize;
@@ -362,18 +236,18 @@ int32 SWorldPartitionEditorGridSpatialHash::PaintGrid(const FGeometry& AllottedG
 				float ColorGradient = bNeedsGradient ? FMath::Min((CellScreenSize.X - DefaultCoordTextSize.X) / GradientDistance, 1.0f) : 1.0f;
 				const FLinearColor CoordTextColor(1.0f, 1.0f, 1.0f, ColorGradient);
 
-				for (int32 y = TopLeftW.Y; y < BottomRightW.Y; y += EditorSpatialHash->CellSize)
+				for (int32 y = TopLeftW.Y; y < BottomRightW.Y; y += EffectiveCellSize)
 				{
-					for (int32 x = TopLeftW.X; x < BottomRightW.X; x += EditorSpatialHash->CellSize)
+					for (int32 x = TopLeftW.X; x < BottomRightW.X; x += EffectiveCellSize)
 					{
-						const FString CoordText(FString::Printf(TEXT("(%d,%d)"), x / EditorSpatialHash->CellSize, y / EditorSpatialHash->CellSize));
+						const FString CoordText(FString::Printf(TEXT("(%d,%d)"), x / EffectiveCellSize, y / EffectiveCellSize));
 						const FVector2D CoordTextSize = FontMeasure->Measure(CoordText, CoordsFont);
 
 						FSlateDrawElement::MakeText(
 							OutDrawElements,
 							++LayerId,
-							AllottedGeometry.ToPaintGeometry(WorldToScreen.TransformPoint(FVector2D(x + EditorSpatialHash->CellSize / 2, y + EditorSpatialHash->CellSize / 2)) - CoordTextSize / 2, FVector2D(1,1)),
-							FString::Printf(TEXT("(%d,%d)"), x / EditorSpatialHash->CellSize, y / EditorSpatialHash->CellSize),
+							AllottedGeometry.ToPaintGeometry(WorldToScreen.TransformPoint(FVector2D(x + EffectiveCellSize / 2, y + EffectiveCellSize / 2)) - CoordTextSize / 2, FVector2D(1,1)),
+							FString::Printf(TEXT("(%d,%d)"), x / EffectiveCellSize, y / EffectiveCellSize),
 							CoordsFont,
 							ESlateDrawEffect::None,
 							CoordTextColor
diff --git a/Engine/Source/Editor/WorldPartitionEditor/Private/WorldPartitionEditorModule.cpp b/Engine/Source/Editor/WorldPartitionEditor/Private/WorldPartitionEditorModule.cpp
--- a/Engine/Source/Editor/WorldPartitionEditor/Private/WorldPartitionEditorModule.cpp
+++ b/Engine/Source/Editor/WorldPartitionEditor/Private/WorldPartitionEditorModule.cpp
@@ -2,23 +2,17 @@
 #include "WorldPartitionEditorModule.h"
 
 #include "WorldPartition/WorldPartition.h"
-#include "WorldPartition/WorldPartitionVolume.h"
-
 #include "WorldPartition/WorldPartitionHLODsBuilder.h"
 #include "WorldPartition/WorldPartitionMiniMapBuilder.h"
-
+#include "WorldPartition/WorldPartitionActorLoaderInterface.h"
 #include "WorldPartition/HLOD/HLODLayerAssetTypeActions.h"
-
 #include "WorldPartition/SWorldPartitionEditor.h"
 #include "WorldPartition/SWorldPartitionEditorGridSpatialHash.h"
-
 #include "WorldPartition/Customizations/WorldPartitionDetailsCustomization.h"
 #include "WorldPartition/Customizations/WorldPartitionRuntimeSpatialHashDetailsCustomization.h"
-
 #include "WorldPartition/SWorldPartitionConvertDialog.h"
 #include "WorldPartition/WorldPartitionConvertOptions.h"
 #include "WorldPartition/WorldPartitionEditorSettings.h"
-
 #include "WorldPartition/HLOD/SWorldPartitionBuildHLODsDialog.h"
 
 #include "LevelEditor.h"
@@ -58,13 +52,23 @@ const FName WorldPartitionEditorTabId("WorldBrowserPartitionEditor");
 
 DEFINE_LOG_CATEGORY_STATIC(LogWorldPartitionEditor, All, All);
 
-// World Partition
-static void OnLoadSelectedWorldPartitionVolumes(TArray<TWeakObjectPtr<AActor>> Volumes)
+static void OnSelectedWorldPartitionVolumesToggleLoading(TArray<TWeakObjectPtr<AActor>> Volumes, bool bLoad)
 {
 	for (TWeakObjectPtr<AActor> Actor: Volumes)
 	{
-		AWorldPartitionVolume* WorldPartitionVolume = CastChecked<AWorldPartitionVolume>(Actor.Get());
-		WorldPartitionVolume->LoadIntersectingCells(true);
+		if (Actor->Implements<UWorldPartitionActorLoaderInterface>())
+		{
+			IWorldPartitionActorLoaderInterface::ILoaderAdapter* LoaderAdapter = Cast<IWorldPartitionActorLoaderInterface>(Actor)->GetLoaderAdapter();
+
+			if (bLoad)
+			{
+				LoaderAdapter->Load();
+			}
+			else
+			{
+				LoaderAdapter->Unload();
+			}
+		}
 	}
 }
 
@@ -73,10 +77,18 @@ static void CreateLevelViewportContextMenuEntries(FMenuBuilder& MenuBuilder, TAr
 	MenuBuilder.BeginSection("WorldPartition", LOCTEXT("WorldPartition", "World Partition"));
 
 	MenuBuilder.AddMenuEntry(
-		LOCTEXT("WorldPartitionLoad", "Load selected world partition volumes"),
-		LOCTEXT("WorldPartitionLoad_Tooltip", "Load selected world partition volumes"),
+		LOCTEXT("WorldPartitionLoad", "Load selected volumes"),
+		LOCTEXT("WorldPartitionLoad_Tooltip", "Load selected volumes"),
+		FSlateIcon(),
+		FExecuteAction::CreateStatic(OnSelectedWorldPartitionVolumesToggleLoading, Volumes, true),
+		NAME_None,
+		EUserInterfaceActionType::Button);
+
+	MenuBuilder.AddMenuEntry(
+		LOCTEXT("WorldPartitionUnload", "Unload selected volumes"),
+		LOCTEXT("WorldPartitionUnload_Tooltip", "Load selected volumes"),
 		FSlateIcon(),
-		FExecuteAction::CreateStatic(OnLoadSelectedWorldPartitionVolumes, Volumes),
+		FExecuteAction::CreateStatic(OnSelectedWorldPartitionVolumesToggleLoading, Volumes, false),
 		NAME_None,
 		EUserInterfaceActionType::Button);
 
@@ -90,7 +102,7 @@ static TSharedRef<FExtender> OnExtendLevelEditorMenu(const TSharedRef<FUICommand
 	TArray<TWeakObjectPtr<AActor> > Volumes;
 	for (AActor* Actor : SelectedActors)
 	{
-		if (Actor->IsA(AWorldPartitionVolume::StaticClass()))
+		if (Actor->Implements<UWorldPartitionActorLoaderInterface>())
 		{
 			Volumes.Add(Actor);
 		}
diff --git a/Engine/Source/Editor/WorldPartitionEditor/Public/WorldPartition/SWorldPartitionEditorGrid2D.h b/Engine/Source/Editor/WorldPartitionEditor/Public/WorldPartition/SWorldPartitionEditorGrid2D.h
--- a/Engine/Source/Editor/WorldPartitionEditor/Public/WorldPartition/SWorldPartitionEditorGrid2D.h
+++ b/Engine/Source/Editor/WorldPartitionEditor/Public/WorldPartition/SWorldPartitionEditorGrid2D.h
@@ -7,9 +7,6 @@
 #include "Layout/ArrangedChildren.h"
 #include "SWorldPartitionEditorGrid.h"
 
-class UWorldPartition2D;
-class UWorldPartitionEditorCell;
-
 class SWorldPartitionEditorGrid2D : public SWorldPartitionEditorGrid
 {
 protected:
@@ -18,9 +15,9 @@ protected:
 	public:
 		FEditorCommands();
 	
-		TSharedPtr<FUICommandInfo> LoadSelectedCells;
-		TSharedPtr<FUICommandInfo> UnloadSelectedCells;
-		TSharedPtr<FUICommandInfo> UnloadAllCells;
+		TSharedPtr<FUICommandInfo> CreateRegionFromSelection;
+		TSharedPtr<FUICommandInfo> LoadSelectedRegions;
+		TSharedPtr<FUICommandInfo> UnloadSelectedRegions;
 		TSharedPtr<FUICommandInfo> MoveCameraHere;
 
 		/**
@@ -35,9 +32,9 @@ public:
 
 	void Construct(const FArguments& InArgs);
 
-	void LoadSelectedCells();
-	void UnloadSelectedCells();
-	void UnloadAllCells();
+	virtual void CreateRegionFromSelection();
+	virtual void LoadSelectedRegions();
+	virtual void UnloadSelectedRegions();
 	void MoveCameraHere();
 
 	virtual FReply OnMouseButtonDown(const FGeometry& MyGeometry, const FPointerEvent& MouseEvent) override;
diff --git a/Engine/Source/Runtime/Engine/Private/LocationVolume.cpp b/Engine/Source/Runtime/Engine/Private/LocationVolume.cpp
new file mode 100644
--- /dev/null
+++ b/Engine/Source/Runtime/Engine/Private/LocationVolume.cpp
@@ -0,0 +1,43 @@
+// Copyright Epic Games, Inc. All Rights Reserved.
+
+#include "LocationVolume.h"
+#include "EngineDefines.h"
+#include "Components/BrushComponent.h"
+#include "WorldPartition/LoaderAdapter/LoaderAdapterActor.h"
+
+ALocationVolume::ALocationVolume(const FObjectInitializer& ObjectInitializer)
+: Super(ObjectInitializer)
+{
+	GetBrushComponent()->SetGenerateOverlapEvents(false);
+
+#if WITH_EDITORONLY_DATA
+	bIsSpatiallyLoaded = false;
+#endif
+
+#if WITH_EDITOR
+	if (!IsTemplate() && GetWorld() && GetWorld()->GetWorldPartition())
+	{
+		WorldPartitionActorLoader = new FLoaderAdapterActor(this);
+	}
+#endif
+}
+
+void ALocationVolume::BeginDestroy()
+{
+#if WITH_EDITOR
+	if (WorldPartitionActorLoader)
+	{
+		delete WorldPartitionActorLoader;
+		WorldPartitionActorLoader = nullptr;
+	}
+#endif
+
+	Super::BeginDestroy();
+}
+
+#if WITH_EDITOR
+IWorldPartitionActorLoaderInterface::ILoaderAdapter* ALocationVolume::GetLoaderAdapter()
+{
+	return WorldPartitionActorLoader;
+}
+#endif
\ No newline at end of file
diff --git a/Engine/Source/Runtime/Engine/Private/WorldPartition/LoaderAdapter/LoaderAdapterActor.cpp b/Engine/Source/Runtime/Engine/Private/WorldPartition/LoaderAdapter/LoaderAdapterActor.cpp
new file mode 100644
--- /dev/null
+++ b/Engine/Source/Runtime/Engine/Private/WorldPartition/LoaderAdapter/LoaderAdapterActor.cpp
@@ -0,0 +1,26 @@
+// Copyright Epic Games, Inc. All Rights Reserved.
+
+#include "WorldPartition/LoaderAdapter/LoaderAdapterActor.h"
+
+#if WITH_EDITOR
+FLoaderAdapterActor::FLoaderAdapterActor(AActor* InActor)
+	: ILoaderAdapterSpatial(InActor->GetWorld())
+	, Actor(InActor)
+{}
+
+TOptional<FBox> FLoaderAdapterActor::GetBoundingBox() const
+{
+	return Actor->GetComponentsBoundingBox();
+}
+
+TOptional<FString> FLoaderAdapterActor::GetLabel() const
+{
+	return Actor->GetActorNameOrLabel();
+}
+
+bool FLoaderAdapterActor::Intersect(const FBox& Box) const
+{
+	// we'll always intersect until we start handling custom brushes
+	return true;
+}
+#endif
diff --git a/Engine/Source/Runtime/Engine/Private/WorldPartition/LoaderAdapter/LoaderAdapterShape.cpp b/Engine/Source/Runtime/Engine/Private/WorldPartition/LoaderAdapter/LoaderAdapterShape.cpp
new file mode 100644
--- /dev/null
+++ b/Engine/Source/Runtime/Engine/Private/WorldPartition/LoaderAdapter/LoaderAdapterShape.cpp
@@ -0,0 +1,26 @@
+// Copyright Epic Games, Inc. All Rights Reserved.
+
+#include "WorldPartition/LoaderAdapter/LoaderAdapterShape.h"
+
+#if WITH_EDITOR
+FLoaderAdapterShape::FLoaderAdapterShape(UWorld* InWorld, const FBox& InBoundingBox, const FString& InLabel)
+	: ILoaderAdapterSpatial(InWorld)
+	, BoundingBox(InBoundingBox)
+	, Label(InLabel)
+{}
+
+TOptional<FBox> FLoaderAdapterShape::GetBoundingBox() const
+{
+	return BoundingBox;
+}
+
+TOptional<FString> FLoaderAdapterShape::GetLabel() const
+{
+	return Label;
+}
+
+bool FLoaderAdapterShape::Intersect(const FBox& Box) const
+{
+	return BoundingBox.Intersect(Box);
+}
+#endif
diff --git a/Engine/Source/Runtime/Engine/Private/WorldPartition/WorldPartition.cpp b/Engine/Source/Runtime/Engine/Private/WorldPartition/WorldPartition.cpp
--- a/Engine/Source/Runtime/Engine/Private/WorldPartition/WorldPartition.cpp
+++ b/Engine/Source/Runtime/Engine/Private/WorldPartition/WorldPartition.cpp
@@ -5,7 +5,6 @@
 =============================================================================*/
 #include "WorldPartition/WorldPartition.h"
 #include "WorldPartition/WorldPartitionLog.h"
-#include "WorldPartition/WorldPartitionEditorCell.h"
 #include "WorldPartition/WorldPartitionRuntimeCell.h"
 #include "WorldPartition/WorldPartitionStreamingPolicy.h"
 #include "WorldPartition/WorldPartitionLevelStreamingPolicy.h"
@@ -45,6 +44,7 @@
 #include "WorldPartition/DataLayer/DataLayerInstance.h"
 #include "WorldPartition/DataLayer/DataLayerSubsystem.h"
 #include "WorldPartition/DataLayer/WorldDataLayers.h"
+#include "WorldPartition/LoaderAdapter/LoaderAdapterShape.h"
 #include "WorldPartition/WorldPartitionActorDescViewProxy.h"
 #include "WorldPartition/HLOD/HLODLayer.h"
 #include "Modules/ModuleManager.h"
@@ -198,59 +198,16 @@ static FAutoConsoleCommand SetLogWorldPartitionVerbosity(
 	})
 );
 
-// Helper class to avoid sending global events until all cells updates are processed.
-struct FWorldPartionCellUpdateContext
+class FLoaderAdapterAlwaysLoadedActors : public FLoaderAdapterShape
 {
-	FWorldPartionCellUpdateContext(UWorldPartition* InWorldPartition)
-		: WorldPartition(InWorldPartition)
-		, bCancelled(false)
+public:
+	FLoaderAdapterAlwaysLoadedActors(UWorld* InWorld)
+		: FLoaderAdapterShape(InWorld, FBox(FVector(-WORLDPARTITION_MAX, -WORLDPARTITION_MAX, -WORLDPARTITION_MAX), FVector(WORLDPARTITION_MAX, WORLDPARTITION_MAX, WORLDPARTITION_MAX)), TEXT("Always Loaded"))
 	{
-		WorldPartition->OnCancelWorldPartitionUpdateEditorCells.AddRaw(this, &FWorldPartionCellUpdateContext::OnCancelUpdateEditorCells);
-
-		UpdatesInProgress++;
-	}
-
-	void OnCancelUpdateEditorCells(UWorldPartition* InWorldPartition)
-	{
-		if (WorldPartition == InWorldPartition)
-		{
-			bCancelled = true;
-		}
-	}
-
-	~FWorldPartionCellUpdateContext()
-	{
-		WorldPartition->OnCancelWorldPartitionUpdateEditorCells.RemoveAll(this);
-
-		UpdatesInProgress--;
-		if (UpdatesInProgress == 0 && !bCancelled)
-		{
-			// @todo_ow: Once Metadata is removed from external actor's package, testing WorldPartition->IsInitialized() won't be necessary anymore.
-			if (WorldPartition->IsInitialized())
-			{
-				if (!IsRunningCommandlet())
-				{
-					GEngine->BroadcastLevelActorListChanged();
-					GEditor->NoteSelectionChange();
-
-					if (WorldPartition->WorldPartitionEditor)
-					{
-						WorldPartition->WorldPartitionEditor->Refresh();
-					}
-
-					GEditor->ResetTransaction(LOCTEXT("LoadingEditorCellsResetTrans", "Editor Cells Loading State Changed"));
-				}
-			}
-		}
+		bIncludeSpatiallyLoadedActors = false;
+		bIncludeNonSpatiallyLoadedActors = true;
 	}
-
-	static int32 UpdatesInProgress;
-	UWorldPartition* WorldPartition;
-	bool bCancelled;
 };
-
-int32 FWorldPartionCellUpdateContext::UpdatesInProgress = 0;
-
 #endif
 
 UWorldPartition::UWorldPartition(const FObjectInitializer& ObjectInitializer)
@@ -418,6 +375,7 @@ void UWorldPartition::Initialize(UWorld* InWorld, const FTransform& InTransform)
 	{
 		// Don't rely on the editor hash for cooking or -game
 		EditorHash = nullptr;
+		AlwaysLoadedActors = nullptr;
 	}
 	else if (bIsEditor)
 	{
@@ -427,6 +385,8 @@ void UWorldPartition::Initialize(UWorld* InWorld, const FTransform& InTransform)
 		check(EditorHash);
 
 		EditorHash->Initialize();
+
+		AlwaysLoadedActors = new FLoaderAdapterAlwaysLoadedActors(World);
 	}
 
 	check(RuntimeHash);
@@ -487,7 +447,7 @@ void UWorldPartition::Initialize(UWorld* InWorld, const FTransform& InTransform)
 		// Also, this will ensure that FindObject will work on an external actor that was loading asynchronously.
 		FlushAsyncLoading();
 
-		// Make sure to preload only AWorldDataLayers actor first (ShouldActorBeLoadedByEditorCells requires it)
+		// Make sure to preload only AWorldDataLayers actor first
 		for (UActorDescContainer::TIterator<> ActorDescIterator(this); ActorDescIterator; ++ActorDescIterator)
 		{
 			if (ActorDescIterator->GetActorNativeClass()->IsChildOf<AWorldDataLayers>())
@@ -532,18 +492,16 @@ void UWorldPartition::Initialize(UWorld* InWorld, const FTransform& InTransform)
 	if (bIsEditor && !bIsCooking)
 	{
 		// Load the always loaded cell, don't call LoadCells to avoid creating a transaction
-		UpdateLoadingEditorCell(EditorHash->GetAlwaysLoadedCell(), true, false);
+		AlwaysLoadedActors->Load();
 
 		// Load more cells depending on the user's settings
 		// Skipped when running from a commandlet
 		if (!IsRunningCommandlet())
 		{
-			// Load last loaded cells
-			if (GetMutableDefault<UWorldPartitionEditorPerProjectUserSettings>()->GetEnableLoadingOfLastLoadedCells())
+			// Load last loaded regions
+			if (GetMutableDefault<UWorldPartitionEditorPerProjectUserSettings>()->GetEnableLoadingOfLastLoadedRegions())
 			{
-				TArray<FName> EditorGridLastLoadedCells = GetMutableDefault<UWorldPartitionEditorPerProjectUserSettings>()->GetEditorGridLoadedCells(InWorld);
-
-				LoadEditorCells(EditorGridLastLoadedCells, /*bIsFromUserChange=*/true);
+				LoadLastLoadedRegions();
 			}
 		}
 		
@@ -609,20 +567,21 @@ void UWorldPartition::Uninitialize()
 		{
 			OnEndPlay();
 		}
-		else
+		
+		if (AlwaysLoadedActors)
 		{
-			if (!IsEngineExitRequested())
+			delete AlwaysLoadedActors;
+			AlwaysLoadedActors = nullptr;
+		}
+
+		if (RegisteredEditorLoaderAdapters.Num())
+		{
+			for (IWorldPartitionActorLoaderInterface::ILoaderAdapter* RegisteredLoaderAdapter : RegisteredEditorLoaderAdapters)
 			{
-				// Unload all Editor cells
-				if (EditorHash)
-				{
-					// @todo_ow: Once Metadata is removed from external actor's package, this won't be necessary anymore.
-					EditorHash->ForEachCell([this](UWorldPartitionEditorCell* Cell)
-					{
-						UpdateLoadingEditorCell(Cell, /*bShouldBeLoaded*/false, /*bIsFromUserChange*/false);
-					});
-				}
+				delete RegisteredLoaderAdapter;
 			}
+
+			RegisteredEditorLoaderAdapters.Empty();
 		}
 
 		WorldDataLayersActor = FWorldPartitionReference();
@@ -672,34 +631,17 @@ const FTransform* UWorldPartition::GetInstanceTransformPtr() const
 void UWorldPartition::OnPostBugItGoCalled(const FVector& Loc, const FRotator& Rot)
 {
 #if WITH_EDITOR
-	if (GetMutableDefault<UWorldPartitionEditorPerProjectUserSettings>()->GetBugItGoLoadCells())
-{
+	if (GetMutableDefault<UWorldPartitionEditorPerProjectUserSettings>()->GetBugItGoLoadRegion())
+	{
 		const FVector LoadExtent(GLoadingRangeBugItGo, GLoadingRangeBugItGo, WORLDPARTITION_MAX);
 		const FBox LoadCellsBox(Loc - LoadExtent, Loc + LoadExtent);
-		LoadEditorCells(LoadCellsBox, false);
+
+		FLoaderAdapterShape* LoaderAdapter = CreateEditorLoaderAdapter<FLoaderAdapterShape>(World, LoadCellsBox, TEXT("BugItGo"));
+		LoaderAdapter->Load();
 	}
 #endif
 }
 
-#if WITH_EDITOR
-void UWorldPartition::OnActorDataLayersEditorLoadingStateChanged(bool bFromUserOperation)
-{
-	auto GetCellsToRefresh = [this](TArray<UWorldPartitionEditorCell*>& CellsToRefresh)
-	{
-		EditorHash->ForEachCell([&CellsToRefresh](UWorldPartitionEditorCell* Cell)
-		{
-			if (Cell->IsLoaded())
-			{
-				CellsToRefresh.Add(Cell);
-			}
-		});
-		return !CellsToRefresh.IsEmpty();
-	};
-
-	UpdateEditorCells(GetCellsToRefresh, /*bIsCellShouldBeLoaded*/true, bFromUserOperation);
-}
-#endif
-
 void UWorldPartition::RegisterDelegates()
 {
 	check(World); 
@@ -716,8 +658,6 @@ void UWorldPartition::RegisterDelegates()
 			FCoreUObjectDelegates::PostReachabilityAnalysis.AddUObject(this, &UWorldPartition::OnGCPostReachabilityAnalysis);
 			GEditor->OnPostBugItGoCalled().AddUObject(this, &UWorldPartition::OnPostBugItGoCalled);
 		}
-
-		FDataLayersEditorBroadcast::Get().OnActorDataLayersEditorLoadingStateChanged().AddUObject(this, &UWorldPartition::OnActorDataLayersEditorLoadingStateChanged);
 	}
 #endif
 
@@ -752,8 +692,6 @@ void UWorldPartition::UnregisterDelegates()
 
 			GEditor->OnPostBugItGoCalled().RemoveAll(this);
 		}
-
-		FDataLayersEditorBroadcast::Get().OnActorDataLayersEditorLoadingStateChanged().RemoveAll(this);
 	}
 #endif
 
@@ -869,267 +807,6 @@ bool UWorldPartition::IsSimulating()
 }
 
 #if WITH_EDITOR
-void UWorldPartition::LoadEditorCells(const FBox& Box, bool bIsFromUserChange)
-{
-	TArray<UWorldPartitionEditorCell*> CellsToLoad;
-	if (EditorHash->GetIntersectingCells(Box, CellsToLoad))
-	{
-		TRACE_CPUPROFILER_EVENT_SCOPE(UWorldPartition::LoadEditorCells);
-
-		FWorldPartionCellUpdateContext CellUpdateContext(this);
-
-		int32 NumActorsToLoad = Algo::TransformAccumulate(CellsToLoad, [](UWorldPartitionEditorCell* Cell) { return Cell->Actors.Num() - Cell->LoadedActors.Num();}, 0);
-
-		FScopedSlowTask SlowTask(NumActorsToLoad, LOCTEXT("LoadingCells", "Loading cells..."));
-		SlowTask.MakeDialog();
-
-		for (UWorldPartitionEditorCell* Cell : CellsToLoad)
-		{
-			SlowTask.EnterProgressFrame(Cell->Actors.Num() - Cell->LoadedActors.Num());
-			UpdateLoadingEditorCell(Cell, true, bIsFromUserChange);
-		}
-	}
-}
-
-void UWorldPartition::LoadEditorCells(const TArray<FName>& CellNames, bool bIsFromUserChange)
-{
-	for (FName CellName : CellNames)
-	{
-		if (UWorldPartitionEditorCell* Cell = FindObject<UWorldPartitionEditorCell>(EditorHash, *CellName.ToString()))
-		{
-			UpdateLoadingEditorCell(Cell, true, bIsFromUserChange);
-		}
-	}
-}
-
-void UWorldPartition::UnloadEditorCells(const FBox& Box, bool bIsFromUserChange)
-{
-	FWorldPartionCellUpdateContext CellUpdateContext(this);
-
-	auto GetCellsToUnload = [this, Box](TArray<UWorldPartitionEditorCell*>& CellsToUnload)
-	{
-		return EditorHash->GetIntersectingCells(Box, CellsToUnload) > 0;
-	};
-
-	UpdateEditorCells(GetCellsToUnload, /*bIsCellShouldBeLoaded*/false, bIsFromUserChange);
-}
-
-bool UWorldPartition::AreEditorCellsLoaded(const FBox& Box)
-{
-	TArray<UWorldPartitionEditorCell*> CellsToLoad;
-	if (EditorHash->GetIntersectingCells(Box, CellsToLoad))
-	{
-		for (UWorldPartitionEditorCell* Cell : CellsToLoad)
-		{
-			if (!Cell->IsLoaded())
-			{
-				return false;
-			}
-		}
-	}
-
-	return true;
-}
-
-bool UWorldPartition::UpdateEditorCells(TFunctionRef<bool(TArray<UWorldPartitionEditorCell*>&)> GetCellsToProcess, bool bIsCellShouldBeLoaded, bool bIsFromUserChange)
-{
-	FWorldPartionCellUpdateContext CellUpdateContext(this);
-
-	TArray<UWorldPartitionEditorCell*> CellsToProcess;
-	if (!GetCellsToProcess(CellsToProcess))
-	{
-		return true;
-	}
-
-	TSet<UPackage*> ModifiedPackages;
-	TMap<FWorldPartitionActorDesc*, int32> UnloadCount;
-
-	for (UWorldPartitionEditorCell* Cell : CellsToProcess)
-	{
-		for (const UWorldPartitionEditorCell::FActorReference& ActorDesc : Cell->LoadedActors)
-		{
-			if (ActorDesc.IsValid())
-			{
-				if (!bIsCellShouldBeLoaded || !ShouldActorBeLoadedByEditorCells(*ActorDesc))
-				{
-					UnloadCount.FindOrAdd(*ActorDesc, 0)++;
-				}
-			}
-		}
-	}
-
-	bool bIsUpdateUnloadingActors = false;
-
-	for (const TPair<FWorldPartitionActorDesc*, int32> Pair : UnloadCount)
-	{
-		FWorldPartitionActorDesc* ActorDesc = Pair.Key;
-
-		// Only prompt if the actor will get unloaded by the unloading cells
-		if (ActorDesc->GetHardRefCount() == Pair.Value)
-		{
-			if (AActor* LoadedActor = ActorDesc->GetActor())
-			{
-				UPackage* ActorPackage = LoadedActor->GetExternalPackage();
-				if (ActorPackage && ActorPackage->IsDirty())
-				{
-					ModifiedPackages.Add(ActorPackage);
-				}
-				bIsUpdateUnloadingActors = true;
-			}
-		}
-	}
-
-	// Make sure we save modified actor packages before unloading
-	FEditorFileUtils::EPromptReturnCode RetCode = FEditorFileUtils::PR_Success;
-
-	// Skip this when running commandlet because MarkPackageDirty() is allowed to dirty packages at loading when running a commandlet.
-	// Usually, the main reason actor packages are dirtied is when RerunConstructionScripts is called on the actor when it is added to the level.
-	if (ModifiedPackages.Num() && !IsRunningCommandlet())
-	{
-		const bool bCheckDirty = false;
-		const bool bAlreadyCheckedOut = false;
-		const bool bCanBeDeclined = true;
-		const bool bPromptToSave = true;
-		const FText Title = LOCTEXT("SaveActorsTitle", "Save Actor(s)");
-		const FText Message = LOCTEXT("SaveActorsMessage", "Save Actor(s) before unloading them.");
-
-		RetCode = FEditorFileUtils::PromptForCheckoutAndSave(ModifiedPackages.Array(), bCheckDirty, bPromptToSave, Title, Message, nullptr, bAlreadyCheckedOut, bCanBeDeclined);
-		if (RetCode == FEditorFileUtils::PR_Cancelled)
-		{
-			OnCancelWorldPartitionUpdateEditorCells.Broadcast(this);
-			return false;
-		}
-
-		check(RetCode != FEditorFileUtils::PR_Failure);
-	}
-
-	if (bIsUpdateUnloadingActors)
-	{
-		GEditor->SelectNone(true, true);
-	}
-
-	// At this point, cells might have changed due to saving deleted actors
-	CellsToProcess.Empty(CellsToProcess.Num());
-	if (!GetCellsToProcess(CellsToProcess))
-	{
-		return true;
-	}
-
-	FScopedSlowTask SlowTask(CellsToProcess.Num(), LOCTEXT("UpdatingCells", "Updating cells..."));
-	SlowTask.MakeDialog();
-
-	for (UWorldPartitionEditorCell* Cell : CellsToProcess)
-	{
-		SlowTask.EnterProgressFrame(1);
-		UpdateLoadingEditorCell(Cell, bIsCellShouldBeLoaded, bIsFromUserChange);
-	}
-
-	return true;
-}
-
-bool UWorldPartition::ShouldActorBeLoadedByEditorCells(const FWorldPartitionActorDesc* ActorDesc) const
-{
-	if (!ActorDesc->ShouldBeLoadedByEditorCells())
-	{
-		return false;
-	}
-
-	if (UDataLayerSubsystem* DataLayerSubsystem = UWorld::GetSubsystem<UDataLayerSubsystem>(GetWorld()))
-	{
-		// Use DataLayers of loaded/dirty Actor if available to handle dirtied actors
-		FWorldPartitionActorViewProxy ActorDescProxy(ActorDesc);
-
-		if (IsRunningCookCommandlet())
-		{
-			// When running cook commandlet, dont allow loading of actors with dynamically loaded data layers
-			for (const FName& DataLayerInstanceName : ActorDescProxy.GetDataLayers())
-			{
-				const UDataLayerInstance* DataLayerInstance = DataLayerSubsystem->GetDataLayerInstance(DataLayerInstanceName);
-				if (DataLayerInstance && DataLayerInstance->IsRuntime())
-				{
-					return false;
-				}
-			}
-		}
-		else
-		{
-			uint32 NumValidLayers = 0;
-			for (const FName& DataLayerInstanceName : ActorDescProxy.GetDataLayers())
-			{
-				if (const UDataLayerInstance* DataLayerInstance = DataLayerSubsystem->GetDataLayerInstance(DataLayerInstanceName))
-				{
-					if (DataLayerInstance->IsEffectiveLoadedInEditor())
-					{
-						return true;
-					}
-					NumValidLayers++;
-				}
-			}
-			return !NumValidLayers;
-		}
-	}
-	
-	return true;
-};
-
-// Loads actors in Editor cell
-void UWorldPartition::UpdateLoadingEditorCell(UWorldPartitionEditorCell* Cell, bool bShouldBeLoaded, bool bFromUserOperation)
-{
-	TRACE_CPUPROFILER_EVENT_SCOPE(UWorldPartition::UpdateLoadingEditorCell);
-	
-	FWorldPartionCellUpdateContext CellUpdateContext(this);
-
-	UE_LOG(LogWorldPartition, Verbose, TEXT("UWorldPartition::UpdateLoadingEditorCell 0x%08llx [%s/%s]"), Cell, bShouldBeLoaded ? TEXT("Load") : TEXT("Unload"), bFromUserOperation ? TEXT("User") : TEXT("Auto"));
-
-	Cell->Modify(false);
-
-	bool bPotentiallyUnloadedActors = false;
-
-	if (!bShouldBeLoaded)
-	{
-		bPotentiallyUnloadedActors = !Cell->LoadedActors.IsEmpty();
-		Cell->LoadedActors.Empty();
-	}
-	else
-	{
-		for (UWorldPartitionEditorCell::FActorHandle& ActorHandle: Cell->Actors)
-		{
-			if (ActorHandle.IsValid())
-			{
-				// Filter actor against DataLayers
-				if (ShouldActorBeLoadedByEditorCells(*ActorHandle))
-				{
-					Cell->LoadedActors.Add(UWorldPartitionEditorCell::FActorReference(ActorHandle.Source, ActorHandle.Handle));
-				}
-				else
-				{
-					// Don't call LoadedActors.Remove(ActorHandle) right away, as this will create a temporary reference and might try to load.
-					for (const UWorldPartitionEditorCell::FActorReference& ActorReference : Cell->LoadedActors)
-					{
-						if ((ActorReference.Source == ActorHandle.Source) && (ActorReference.Handle == ActorHandle.Handle))
-						{
-							Cell->LoadedActors.Remove(UWorldPartitionEditorCell::FActorReference(ActorHandle.Source, ActorHandle.Handle));
-							bPotentiallyUnloadedActors = true;
-							break;
-						}
-					}
-				}
-			}
-		}
-	}
-
-	if (Cell->IsLoaded() != bShouldBeLoaded)
-	{
-		Cell->SetLoaded(bShouldBeLoaded, bShouldBeLoaded && bFromUserOperation);
-	}
-
-	if (bPotentiallyUnloadedActors)
-	{
-		bForceGarbageCollection = true;
-		bForceGarbageCollectionPurge = true;
-	}
-}
-
 void UWorldPartition::OnActorDescAdded(FWorldPartitionActorDesc* NewActorDesc)
 {
 	Super::OnActorDescAdded(NewActorDesc);
@@ -1261,7 +938,7 @@ void UWorldPartition::OnEnableStreamingChanged()
 
 	if (!IsStreamingEnabled())
 	{
-		UpdateLoadingEditorCell(EditorHash->GetAlwaysLoadedCell(), true, false);
+		AlwaysLoadedActors->Load();
 	}
 
 	if (WorldPartitionEditor)
@@ -1274,6 +951,7 @@ void UWorldPartition::HashActorDesc(FWorldPartitionActorDesc* ActorDesc)
 {
 	check(ActorDesc);
 	check(EditorHash);
+
 	FWorldPartitionHandle ActorHandle(this, ActorDesc->GetGuid());
 	EditorHash->HashActor(ActorHandle);
 
@@ -1284,6 +962,7 @@ void UWorldPartition::UnhashActorDesc(FWorldPartitionActorDesc* ActorDesc)
 {
 	check(ActorDesc);
 	check(EditorHash);
+
 	FWorldPartitionHandle ActorHandle(this, ActorDesc->GetGuid());
 	EditorHash->UnhashActor(ActorHandle);
 }
@@ -1475,37 +1154,26 @@ bool UWorldPartition::FinalizeGeneratorPackageForCook(const TArray<ICookPackageS
 	return false;
 }
 
-TArray<FName> UWorldPartition::GetUserLoadedEditorGridCells() const
+TArray<FBox> UWorldPartition::GetUserLoadedEditorGridRegions() const
 {
-	// Save last loaded cells settings
-	TArray<FName> LastEditorGridLoadedCells = GetMutableDefault<UWorldPartitionEditorPerProjectUserSettings>()->GetEditorGridLoadedCells(GetWorld());
+	TArray<FBox> Result;
 
-	TArray<FName> EditorGridLastLoadedCells;
-	EditorHash->ForEachCell([this, &LastEditorGridLoadedCells, &EditorGridLastLoadedCells](UWorldPartitionEditorCell* Cell)
+	for (IWorldPartitionActorLoaderInterface::ILoaderAdapter* LoaderAdapter : RegisteredEditorLoaderAdapters)
 	{
-		FName CellName = Cell->GetFName();
-
-		if (Cell != EditorHash->GetAlwaysLoadedCell())
+		if (LoaderAdapter->GetUserCreated())
 		{
-			if (Cell->IsLoaded() && Cell->IsLoadedChangedByUserOperation())
-			{
-				EditorGridLastLoadedCells.Add(CellName);
-			}
-			else if (!Cell->IsLoaded() && !Cell->IsLoadedChangedByUserOperation() && LastEditorGridLoadedCells.Contains(CellName))
-			{
-				EditorGridLastLoadedCells.Add(CellName);
-			}
+			Result.Add(*LoaderAdapter->GetBoundingBox());
 		}
-	});
+	}
 
-	return EditorGridLastLoadedCells;
+	return Result;
 }
 
 void UWorldPartition::SavePerUserSettings()
 {
 	if (GIsEditor && !World->IsGameWorld() && !IsRunningCommandlet() && !IsEngineExitRequested())
 	{
-		GetMutableDefault<UWorldPartitionEditorPerProjectUserSettings>()->SetEditorGridLoadedCells(GetWorld(), GetUserLoadedEditorGridCells());
+		GetMutableDefault<UWorldPartitionEditorPerProjectUserSettings>()->SetEditorGridLoadedRegions(GetWorld(), GetUserLoadedEditorGridRegions());
 	}
 }
 
@@ -1623,6 +1291,18 @@ bool UWorldPartition::IsActorPinned(const FGuid& ActorGuid) const
 	return PinnedActors.Contains(ActorGuid);
 }
 
+void UWorldPartition::LoadLastLoadedRegions()
+{
+	TArray<FBox> EditorGridLastLoadedRegions = GetMutableDefault<UWorldPartitionEditorPerProjectUserSettings>()->GetEditorGridLoadedRegions(World);
+
+	for (const FBox& EditorGridLastLoadedRegion : EditorGridLastLoadedRegions)
+	{
+		FLoaderAdapterShape* LoaderAdapter = CreateEditorLoaderAdapter<FLoaderAdapterShape>(World, EditorGridLastLoadedRegion, TEXT("Last Loaded Region"));
+		LoaderAdapter->SetUserCreated(true);
+		LoaderAdapter->Load();
+	}
+}
+
 void UWorldPartition::OnWorldRenamed()
 {
 	UActorDescContainer::OnWorldRenamed();
diff --git a/Engine/Source/Runtime/Engine/Private/WorldPartition/WorldPartitionActorDesc.cpp b/Engine/Source/Runtime/Engine/Private/WorldPartition/WorldPartitionActorDesc.cpp
--- a/Engine/Source/Runtime/Engine/Private/WorldPartition/WorldPartitionActorDesc.cpp
+++ b/Engine/Source/Runtime/Engine/Private/WorldPartition/WorldPartitionActorDesc.cpp
@@ -25,15 +25,12 @@
 #endif
 
 #if WITH_EDITOR
-uint32 FWorldPartitionActorDesc::GlobalTag = 0;
-
 FWorldPartitionActorDesc::FWorldPartitionActorDesc()
 	: bIsUsingDataLayerAsset(false)
 	, SoftRefCount(0)
 	, HardRefCount(0)
 	, Container(nullptr)
 	, bIsForcedNonSpatiallyLoaded(false)
-	, Tag(0)
 {}
 
 void FWorldPartitionActorDesc::Init(const AActor* InActor)
@@ -237,18 +234,23 @@ void FWorldPartitionActorDesc::TransformInstance(const FString& From, const FStr
 
 FString FWorldPartitionActorDesc::ToString() const
 {
+	auto GetBoolStr = [](bool bValue) -> const TCHAR*
+	{
+		return bValue ? TEXT("true") : TEXT("false");
+	};
+
 	return FString::Printf(
 		TEXT("Guid:%s BaseClass:%s NativeClass:%s Name:%s SpatiallyLoaded:%s Bounds:%s RuntimeGrid:%s EditorOnly:%s LevelBoundsRelevant:%s HLODRelevant:%s FolderPath:%s FolderGuid:%s Parent:%s"), 
 		*Guid.ToString(), 
 		*BaseClass.ToString(), 
 		*NativeClass.ToString(), 
 		*GetActorName().ToString(),
-		bIsSpatiallyLoaded ? TEXT("true") : TEXT("false"),
+		GetBoolStr(bIsSpatiallyLoaded),
 		*GetBounds().ToString(),
 		*RuntimeGrid.ToString(),
-		bActorIsEditorOnly ? TEXT("true") : TEXT("false"),
-		bLevelBoundsRelevant ? TEXT("true") : TEXT("false"),
-		bActorIsHLODRelevant ? TEXT("true") : TEXT("false"),
+		GetBoolStr(bActorIsEditorOnly),
+		GetBoolStr(bLevelBoundsRelevant),
+		GetBoolStr(bActorIsHLODRelevant),
 		*FolderPath.ToString(),
 		*FolderGuid.ToString(),
 		*ParentActor.ToString()
diff --git a/Engine/Source/Runtime/Engine/Private/WorldPartition/WorldPartitionActorDescView.cpp b/Engine/Source/Runtime/Engine/Private/WorldPartition/WorldPartitionActorDescView.cpp
--- a/Engine/Source/Runtime/Engine/Private/WorldPartition/WorldPartitionActorDescView.cpp
+++ b/Engine/Source/Runtime/Engine/Private/WorldPartition/WorldPartitionActorDescView.cpp
@@ -140,11 +140,6 @@ FString FWorldPartitionActorDescView::ToString() const
 	return ActorDesc->ToString();
 }
 
-uint32 FWorldPartitionActorDescView::GetTag() const
-{
-	return ActorDesc->Tag;
-}
-
 const FGuid& FWorldPartitionActorDescView::GetParentActor() const
 {
 	return ActorDesc->GetParentActor();
diff --git a/Engine/Source/Runtime/Engine/Private/WorldPartition/WorldPartitionActorLoaderInterface.cpp b/Engine/Source/Runtime/Engine/Private/WorldPartition/WorldPartitionActorLoaderInterface.cpp
new file mode 100644
--- /dev/null
+++ b/Engine/Source/Runtime/Engine/Private/WorldPartition/WorldPartitionActorLoaderInterface.cpp
@@ -0,0 +1,327 @@
+// Copyright Epic Games, Inc. All Rights Reserved.
+
+#include "WorldPartition/WorldPartitionActorLoaderInterface.h"
+
+#if WITH_EDITOR
+#include "Editor.h"
+#include "Editor/EditorEngine.h"
+#include "FileHelpers.h"
+#include "Engine/World.h"
+#include "Misc/ScopedSlowTask.h"
+#include "WorldPartition/WorldPartition.h"
+#include "WorldPartition/WorldPartitionEditorHash.h"
+#include "WorldPartition/WorldPartitionActorDescViewProxy.h"
+#include "WorldPartition/DataLayer/WorldDataLayers.h"
+#include "WorldPartition/DataLayer/DataLayerSubsystem.h"
+#endif
+
+#define LOCTEXT_NAMESPACE "WorldPartition"
+
+UWorldPartitionActorLoaderInterface::UWorldPartitionActorLoaderInterface(const FObjectInitializer& ObjectInitializer)
+	: Super(ObjectInitializer)
+{}
+
+#if WITH_EDITOR
+IWorldPartitionActorLoaderInterface::ILoaderAdapter::ILoaderAdapter(UWorld* InWorld)
+	: World(InWorld)
+	, bLoaded(false)
+	, bUserCreated(false)
+{}
+
+IWorldPartitionActorLoaderInterface::ILoaderAdapter::~ILoaderAdapter()
+{
+	Unload();
+}
+
+bool IWorldPartitionActorLoaderInterface::ILoaderAdapter::Load()
+{
+	if (!bLoaded)
+	{
+		bLoaded = true;
+		RefreshLoadedState();
+		FDataLayersEditorBroadcast::Get().OnActorDataLayersEditorLoadingStateChanged().AddRaw(this, &IWorldPartitionActorLoaderInterface::ILoaderAdapter::OnActorDataLayersEditorLoadingStateChanged);
+	}
+
+	return true;
+}
+
+bool IWorldPartitionActorLoaderInterface::ILoaderAdapter::Unload()
+{
+	if (bLoaded && !IsEngineExitRequested())
+	{
+		if (UWorldPartition* WorldPartition = World->GetWorldPartition())
+		{
+			TArray<FWorldPartitionHandle> ActorsToUnload;
+			for (const TPair<FGuid, TMap<FGuid, FWorldPartitionReference>>& ActorRefToUnload : ActorReferences)
+			{
+				ActorsToUnload.Emplace(WorldPartition, ActorRefToUnload.Key);
+			}
+			
+			if (!AllowUnloadingActors(ActorsToUnload))
+			{
+				return false;
+			}
+
+			FScopedSlowTask SlowTask(1, LOCTEXT("UpdatingLoading", "Updating loading..."));
+			SlowTask.MakeDialog();
+
+			FDataLayersEditorBroadcast::Get().OnActorDataLayersEditorLoadingStateChanged().RemoveAll(this);
+			ActorReferences.Empty();
+			bLoaded = false;
+
+			PostLoadedStateChanged(true);
+
+			SlowTask.EnterProgressFrame(1);
+		}
+	}
+
+	return true;
+}
+
+bool IWorldPartitionActorLoaderInterface::ILoaderAdapter::IsLoaded() const
+{
+	return bLoaded;
+}
+
+bool IWorldPartitionActorLoaderInterface::ILoaderAdapter::RefreshLoadedState()
+{
+	if (bLoaded)
+	{
+		if (UWorldPartition* WorldPartition = World->GetWorldPartition())
+		{
+			TArray<FWorldPartitionHandle> ActorsToLoad;
+			TArray<FWorldPartitionHandle> ActorsToUnload;			
+			ForEachActor([this, &ActorsToLoad, &ActorsToUnload](const FWorldPartitionHandle& Actor)
+			{
+				if (ShouldActorBeLoaded(Actor))
+				{
+					if (!ActorReferences.Contains(Actor->GetGuid()))
+					{
+						ActorsToLoad.Add(Actor);
+					}
+				}
+				else if (ActorReferences.Contains(Actor->GetGuid()))
+				{
+					ActorsToUnload.Add(Actor);
+				}
+			});
+
+			if (!AllowUnloadingActors(ActorsToUnload))
+			{
+				return false;
+			}
+
+			FScopedSlowTask SlowTask(ActorsToLoad.Num() + ActorsToUnload.Num(), LOCTEXT("UpdatingLoading", "Updating loading..."));
+			SlowTask.MakeDialog();
+
+			for (FWorldPartitionHandle& ActorToLoad : ActorsToLoad)
+			{
+				SlowTask.EnterProgressFrame(1);
+				FWorldPartitionHandle ActorHandle(WorldPartition, ActorToLoad->GetGuid());
+				AddReferenceToActor(ActorHandle);
+			}
+
+			for (FWorldPartitionHandle& ActorToUnload : ActorsToUnload)
+			{
+				SlowTask.EnterProgressFrame(1);
+				ActorReferences.Remove(ActorToUnload->GetGuid());
+			}
+
+			PostLoadedStateChanged(ActorsToUnload.Num() > 0);
+		}
+	}
+
+	return true;
+}
+
+bool IWorldPartitionActorLoaderInterface::ILoaderAdapter::AllowUnloadingActors(const TArray<FWorldPartitionHandle>& ActorsToUnload) const
+{
+	if (IsRunningCommandlet())
+	{
+		return true;
+	}
+
+	TSet<UPackage*> ModifiedPackages;
+	TMap<FWorldPartitionHandle, int32> UnloadCount;
+
+	// Count the number of times we are referencing each actor to unload
+	for (const FWorldPartitionHandle& ActorToUnload : ActorsToUnload)
+	{
+		const TMap<FGuid, FWorldPartitionReference>& ActorToUnloadRefs = ActorReferences.FindChecked(ActorToUnload->GetGuid());
+
+		for (const TPair<FGuid, FWorldPartitionReference>& ActorToUnloadRef : ActorToUnloadRefs)
+		{
+			UnloadCount.FindOrAdd(ActorToUnloadRef.Value, 0)++;
+		}
+	}
+
+	for (const TPair<FWorldPartitionHandle, int32>& Pair : UnloadCount)
+	{
+		const FWorldPartitionHandle& ActorHandle = Pair.Key;
+
+		// Only prompt if the actor will get unloaded by the unloading cells
+		if (ActorHandle->GetHardRefCount() == Pair.Value)
+		{
+			if (AActor* LoadedActor = ActorHandle->GetActor())
+			{
+				if (UPackage* ActorPackage = LoadedActor->GetExternalPackage(); ActorPackage && ActorPackage->IsDirty())
+				{
+					ModifiedPackages.Add(ActorPackage);
+				}
+			}
+		}
+	}
+
+	if (ModifiedPackages.Num())
+	{
+		const FText Title = LOCTEXT("SaveActorsTitle", "Save Actor(s)");
+		const FText Message = LOCTEXT("SaveActorsMessage", "Save Actor(s) before unloading them.");
+
+		FEditorFileUtils::EPromptReturnCode RetCode = FEditorFileUtils::PromptForCheckoutAndSave(ModifiedPackages.Array(), false, true, Title, Message, nullptr, false, true);
+		check(RetCode != FEditorFileUtils::PR_Failure);
+
+		if (RetCode == FEditorFileUtils::PR_Cancelled)
+		{
+			return false;
+		}		
+	}
+
+	return true;
+}
+
+bool IWorldPartitionActorLoaderInterface::ILoaderAdapter::ShouldActorBeLoaded(const FWorldPartitionHandle& Actor) const
+{
+	check(Actor.IsValid());
+
+	if (!Actor->ShouldBeLoadedByEditor())
+	{
+		return false;
+	}
+
+	if (UDataLayerSubsystem* DataLayerSubsystem = UWorld::GetSubsystem<UDataLayerSubsystem>(World))
+	{
+		FWorldPartitionActorViewProxy ActorDescProxy(*Actor);
+
+		if (IsRunningCookCommandlet())
+		{
+			// When running cook commandlet, dont allow loading of actors with runtime loaded data layers
+			for (const FName& DataLayerInstanceName : ActorDescProxy.GetDataLayers())
+			{
+				const UDataLayerInstance* DataLayerInstance = DataLayerSubsystem->GetDataLayerInstance(DataLayerInstanceName);
+				if (DataLayerInstance && DataLayerInstance->IsRuntime())
+				{
+					return false;
+				}
+			}
+		}
+		else
+		{
+			uint32 NumValidLayers = 0;
+			for (const FName& DataLayerInstanceName : ActorDescProxy.GetDataLayers())
+			{
+				if (const UDataLayerInstance* DataLayerInstance = DataLayerSubsystem->GetDataLayerInstance(DataLayerInstanceName))
+				{
+					if (DataLayerInstance->IsEffectiveLoadedInEditor())
+					{
+						return true;
+					}
+					NumValidLayers++;
+				}
+			}
+			return !NumValidLayers;
+		}
+	}
+
+	return true;
+};
+
+void IWorldPartitionActorLoaderInterface::ILoaderAdapter::PostLoadedStateChanged(bool bUnloadedActors)
+{
+	if (!IsRunningCommandlet())
+	{
+		if (bUnloadedActors)
+		{
+			GEditor->SelectNone(true, true);
+		}
+
+		GEngine->BroadcastLevelActorListChanged();
+		GEditor->NoteSelectionChange();
+
+		GEditor->ResetTransaction(LOCTEXT("LoadingEditorActorResetTrans", "Editor Actors Loading State Changed"));
+
+		if (bUnloadedActors)
+		{
+			GEngine->ForceGarbageCollection(true);
+		}
+	}
+}
+
+void IWorldPartitionActorLoaderInterface::ILoaderAdapter::AddReferenceToActor(FWorldPartitionHandle& ActorHandle)
+{
+	TFunction<void(const FWorldPartitionHandle&, TMap<FGuid, FWorldPartitionReference>&)> AddReferences = [this, &ActorHandle, &AddReferences](const FWorldPartitionHandle& Handle, TMap<FGuid, FWorldPartitionReference>& ReferenceMap)
+	{
+		if (!ReferenceMap.Contains(Handle->GetGuid()))
+		{
+			ReferenceMap.Emplace(Handle->GetGuid(), Handle);
+			
+			for (const FGuid& ReferencedActorGuid : ActorHandle->GetReferences())
+			{
+				FWorldPartitionHandle ReferenceActorHandle(ActorHandle->GetContainer(), ReferencedActorGuid);
+
+				if (ReferenceActorHandle.IsValid())
+				{
+					AddReferences(ReferenceActorHandle, ReferenceMap);
+				}
+			}
+		}
+	};
+
+	AddReferences(ActorHandle, ActorReferences.Emplace(ActorHandle->GetGuid()));
+}
+
+void IWorldPartitionActorLoaderInterface::ILoaderAdapter::OnActorDataLayersEditorLoadingStateChanged(bool bFromUserOperation)
+{
+	if (!RefreshLoadedState())
+	{
+		// Temporary measure: reset transactions if the user cancels saving to avoid inconsistencies. The next step is to change that
+		// behavior in order to keep dirty actors in memory instead of asking to save them when unloading, in preparation of in editor
+		// camera loading.
+		GEditor->ResetTransaction(LOCTEXT("LoadingEditorActorResetTrans", "Editor Actors Loading State Changed"));
+	}
+}
+
+IWorldPartitionActorLoaderInterface::FLoaderAdapterList::FLoaderAdapterList(UWorld* InWorld)
+	: ILoaderAdapter(InWorld)
+{}
+
+void IWorldPartitionActorLoaderInterface::FLoaderAdapterList::ForEachActor(TFunctionRef<void(const FWorldPartitionHandle&)> InOperation) const
+{
+	for (const FWorldPartitionHandle& Actor : Actors)
+	{
+		InOperation(Actor);
+	}
+}
+
+IWorldPartitionActorLoaderInterface::ILoaderAdapterSpatial::ILoaderAdapterSpatial(UWorld* InWorld)
+	: ILoaderAdapter(InWorld)
+	, bIncludeSpatiallyLoadedActors(true)
+	, bIncludeNonSpatiallyLoadedActors(false)
+{}
+
+void IWorldPartitionActorLoaderInterface::ILoaderAdapterSpatial::ForEachActor(TFunctionRef<void(const FWorldPartitionHandle&)> InOperation) const
+{
+	if (UWorldPartition* WorldPartition = World->GetWorldPartition())
+	{
+		WorldPartition->EditorHash->ForEachIntersectingActor(*GetBoundingBox(), [this, WorldPartition, &InOperation](FWorldPartitionActorDesc* ActorDesc)
+		{
+			if (Intersect(ActorDesc->GetBounds()))
+			{
+				FWorldPartitionHandle ActorHandle(WorldPartition, ActorDesc->GetGuid());
+				InOperation(ActorHandle);
+			}
+		}, bIncludeSpatiallyLoadedActors, bIncludeNonSpatiallyLoadedActors);
+	}
+}
+#endif
+
+#undef LOCTEXT_NAMESPACE
\ No newline at end of file
diff --git a/Engine/Source/Runtime/Engine/Private/WorldPartition/WorldPartitionEditorCell.cpp b/Engine/Source/Runtime/Engine/Private/WorldPartition/WorldPartitionEditorCell.cpp
deleted file mode 100644
diff --git a/Engine/Source/Runtime/Engine/Private/WorldPartition/WorldPartitionEditorPerProjectUserSettings.cpp b/Engine/Source/Runtime/Engine/Private/WorldPartition/WorldPartitionEditorPerProjectUserSettings.cpp
--- a/Engine/Source/Runtime/Engine/Private/WorldPartition/WorldPartitionEditorPerProjectUserSettings.cpp
+++ b/Engine/Source/Runtime/Engine/Private/WorldPartition/WorldPartitionEditorPerProjectUserSettings.cpp
@@ -19,25 +19,25 @@ void UWorldPartitionEditorPerProjectUserSettings::SetWorldDataLayersNonDefaultEd
 	}
 }
 
-void UWorldPartitionEditorPerProjectUserSettings::SetEditorGridLoadedCells(UWorld* InWorld, const TArray<FName>& InEditorGridLoadedCells)
+void UWorldPartitionEditorPerProjectUserSettings::SetEditorGridLoadedRegions(UWorld* InWorld, const TArray<FBox>& InEditorGridLoadedRegions)
 {
 	if (ShouldSaveSettings(InWorld))
 	{
 		FWorldPartitionPerWorldSettings& PerWorldSettings = PerWorldEditorSettings.FindOrAdd(TSoftObjectPtr<UWorld>(InWorld));
-		PerWorldSettings.LoadedEditorGridCells = InEditorGridLoadedCells;
+		PerWorldSettings.LoadedEditorGridRegions = InEditorGridLoadedRegions;
 		
 		SaveConfig();
 	}
 }
 
-TArray<FName> UWorldPartitionEditorPerProjectUserSettings::GetEditorGridLoadedCells(UWorld* InWorld) const
+TArray<FBox> UWorldPartitionEditorPerProjectUserSettings::GetEditorGridLoadedRegions(UWorld* InWorld) const
 {
 	if (const FWorldPartitionPerWorldSettings* PerWorldSettings = GetWorldPartitionPerWorldSettings(InWorld))
 	{
-		return PerWorldSettings->LoadedEditorGridCells;
+		return PerWorldSettings->LoadedEditorGridRegions;
 	}
 
-	return TArray<FName>();
+	return TArray<FBox>();
 }
 
 TArray<FName> UWorldPartitionEditorPerProjectUserSettings::GetWorldDataLayersNotLoadedInEditor(UWorld* InWorld) const
diff --git a/Engine/Source/Runtime/Engine/Private/WorldPartition/WorldPartitionEditorSpatialHash.cpp b/Engine/Source/Runtime/Engine/Private/WorldPartition/WorldPartitionEditorSpatialHash.cpp
--- a/Engine/Source/Runtime/Engine/Private/WorldPartition/WorldPartitionEditorSpatialHash.cpp
+++ b/Engine/Source/Runtime/Engine/Private/WorldPartition/WorldPartitionEditorSpatialHash.cpp
@@ -2,7 +2,6 @@
 
 #include "WorldPartition/WorldPartitionEditorSpatialHash.h"
 #include "WorldPartition/WorldPartition.h"
-#include "WorldPartition/WorldPartitionEditorCell.h"
 #include "WorldPartition/WorldPartitionActorDesc.h"
 #include "WorldPartition/WorldPartitionActorDesc.h"
 #include "WorldPartition/WorldPartitionEditorPerProjectUserSettings.h"
@@ -18,7 +17,6 @@ UWorldPartitionEditorSpatialHash::UWorldPartitionEditorSpatialHash(const FObject
 #if WITH_EDITOR
 	, CellSize(12800)
 	, bBoundsDirty(false)
-	, AlwaysLoadedCell(nullptr)
 #endif
 {}
 
@@ -27,7 +25,7 @@ void UWorldPartitionEditorSpatialHash::Initialize()
 {
 	check(!AlwaysLoadedCell);
 
-	AlwaysLoadedCell = NewObject<UWorldPartitionEditorCell>(this, TEXT("AlwaysLoadedCell"), RF_Transient);
+	AlwaysLoadedCell = MakeUnique<FCell>();
 	AlwaysLoadedCell->Bounds.Init();
 }
 
@@ -49,21 +47,20 @@ void UWorldPartitionEditorSpatialHash::Tick(float DeltaSeconds)
 	if (bBoundsDirty)
 	{
 		FBox NewBounds(ForceInit);
-		ForEachCell([&](UWorldPartitionEditorCell* Cell) { NewBounds += Cell->Bounds; });
+		for (FCell* Cell: Cells)
+		{
+			NewBounds += Cell->Bounds;
+		}
 
 		const int32 OldLevel = GetLevelForBox(Bounds);
+		check(OldLevel == HashLevels.Num() - 1);
+
 		const int32 NewLevel = GetLevelForBox(NewBounds);
-		check(NewLevel <= OldLevel);
+		check(NewLevel <= OldLevel);		
 
 		if (NewLevel < OldLevel)
 		{
-			for (int32 Level=NewLevel+1; Level<=OldLevel; Level++)
-			{
-				ForEachIntersectingCells(Bounds, Level, [&](const FCellCoord& CellCoord)
-				{
-					HashNodes.Remove(CellCoord);
-				});
-			}
+			HashLevels.SetNum(NewLevel + 1);
 		}
 
 		Bounds = NewBounds;
@@ -71,105 +68,88 @@ void UWorldPartitionEditorSpatialHash::Tick(float DeltaSeconds)
 	}
 
 	if (CVarEnableSpatialHashValidation.GetValueOnAnyThread())
-	{	
-		const int32 CurrentLevel = GetLevelForBox(Bounds);
-		ForEachIntersectingCells(Bounds, 0, [&](const FCellCoord& CellCoord)
+	{
+		if (Bounds.IsValid)
 		{
-			UWorldPartitionEditorCell* EditorCell = nullptr;
-			if (UWorldPartitionEditorCell** EditorCellPtr = HashCells.Find(CellCoord))
+			const int32 CurrentLevel = GetLevelForBox(Bounds);
+			check(CurrentLevel == HashLevels.Num() - 1);
+
+			for (int32 HashLevel = 0; HashLevel < HashLevels.Num() - 1; HashLevel++)
 			{
-				if ((*EditorCellPtr)->IsLoaded())
+				for (auto& HashLevelPair : HashLevels[HashLevel])
 				{
-					FCellCoord LevelCellCoord = CellCoord;
-					while (LevelCellCoord.Level < CurrentLevel)
-					{
-						const uint32 ChildIndex = LevelCellCoord.GetChildIndex();
-						LevelCellCoord = LevelCellCoord.GetParentCellCoord();
-						const FCellNode& CellNode = HashNodes.FindChecked(LevelCellCoord);
-						check(CellNode.HasChildNode(ChildIndex));
-					}
-				}
-			}
-		});
-	}
-}
-
-FBox UWorldPartitionEditorSpatialHash::GetActorBounds(const FWorldPartitionHandle& InActorHandle) const
-{
-	FBox ActorBounds;
-	if (InActorHandle->GetIsSpatiallyLoaded())
-	{
-		ActorBounds = InActorHandle->GetBounds();
-	}
+					const FCellCoord CellCoord = HashLevelPair.Key;
+					check(CellCoord.Level == HashLevel);
 
-	check(ActorBounds.IsValid);
-	return ActorBounds;
-}
+					const uint32 ChildIndex = CellCoord.GetChildIndex();
+					const FCellCoord ParentCellCoord = CellCoord.GetParentCellCoord();
+					check(ParentCellCoord.Level == HashLevel + 1);
 
-bool UWorldPartitionEditorSpatialHash::IsActorAlwaysLoaded(const FWorldPartitionHandle& InActorHandle) const
-{
-	if (!InActorHandle->GetIsSpatiallyLoaded())
-	{
-		return true;
+					const FCellNodeElement& ParentCellNodeElement = HashLevels[ParentCellCoord.Level].FindChecked(ParentCellCoord);
+					const FCellNode& ParentCellNode = ParentCellNodeElement.Key;
+					check(ParentCellNode.HasChildNode(ChildIndex));
+				}
+			}
+		}
 	}
-
-	// If an actor covers more that 4 levels in the octree (which means 32K cells), treat it as always loaded
-	const FBox ActorBounds = GetActorBounds(InActorHandle);
-	const int32 ActorLevel = GetLevelForBox(ActorBounds);
-	return (ActorLevel > 4);
 }
 
 void UWorldPartitionEditorSpatialHash::HashActor(FWorldPartitionHandle& InActorHandle)
 {
 	check(InActorHandle.IsValid());
 
-	if (IsActorAlwaysLoaded(InActorHandle))
+	if (!InActorHandle->GetIsSpatiallyLoaded())
 	{
-		AlwaysLoadedCell->AddActor(InActorHandle);
+		AlwaysLoadedCell->Actors.Add(InActorHandle);
 	}
 	else
 	{
-		const FBox ActorBounds = GetActorBounds(InActorHandle);
+		const FBox ActorBounds = InActorHandle->GetBounds();
 		const int32 CurrentLevel = GetLevelForBox(Bounds);
+		const int32 ActorLevel = GetLevelForBox(ActorBounds);
 
-		ForEachIntersectingCells(ActorBounds, 0, [&](const FCellCoord& CellCoord)
+		if (HashLevels.Num() <= ActorLevel)
 		{
-			UWorldPartitionEditorCell* EditorCell = nullptr;
-			if (UWorldPartitionEditorCell** EditorCellPtr = HashCells.Find(CellCoord))
-			{
-				EditorCell = *EditorCellPtr;
-			}
-			else
+			HashLevels.AddDefaulted(ActorLevel - HashLevels.Num() + 1);
+		}
+
+		ForEachIntersectingCells(ActorBounds, ActorLevel, [&](const FCellCoord& CellCoord)
+		{
+			FCellNodeElement& CellNodeElement = HashLevels[CellCoord.Level].FindOrAdd(CellCoord);
+
+			TUniquePtr<FCell>& Cell = CellNodeElement.Value;
+
+			if (!Cell)
 			{
-				EditorCell = NewObject<UWorldPartitionEditorCell>(this, *FString::Printf(TEXT("EditorCell_S%d_X%lld_Y%lld_Z%lld"), CellSize, CellCoord.X, CellCoord.Y, CellCoord.Z), RF_Transient);
-				EditorCell->Bounds = GetCellBounds(CellCoord);
+				Cell = MakeUnique<FCell>();
+				Cell->Bounds = GetCellBounds(CellCoord);
 
-				Cells.Add(EditorCell);
-				HashCells.Add(CellCoord, EditorCell);
+				Cells.Add(Cell.Get());
 
-				Bounds += EditorCell->Bounds;
+				// Increment spatial structure bounds
+				Bounds += Cell->Bounds;
 
-				// Update parent cells
-				FCellCoord LevelCellCoord = CellCoord;
-				while (LevelCellCoord.Level < CurrentLevel)
+				// Update parent nodes
+				FCellCoord CurrentCellCoord = CellCoord;
+				while (CurrentCellCoord.Level < CurrentLevel)
 				{
-					const uint32 ChildIndex = LevelCellCoord.GetChildIndex();
+					const uint32 ChildIndex = CurrentCellCoord.GetChildIndex();
+					CurrentCellCoord = CurrentCellCoord.GetParentCellCoord();
 
-					LevelCellCoord = LevelCellCoord.GetParentCellCoord();
+					FCellNodeElement& ParentCellNodeElement = HashLevels[CurrentCellCoord.Level].FindOrAdd(CurrentCellCoord);
+					FCellNode& ParentCellNode = ParentCellNodeElement.Key;
 
-					FCellNode& CellNode = HashNodes.FindOrAdd(LevelCellCoord);
-
-					if (CellNode.HasChildNode(ChildIndex))
+					if (ParentCellNode.HasChildNode(ChildIndex))
 					{
 						break;
 					}
 
-					CellNode.AddChildNode(ChildIndex);
+					ParentCellNode.AddChildNode(ChildIndex);
 				}
 			}
 
-			check(EditorCell);
-			EditorCell->AddActor(InActorHandle);
+			check(Cell);
+			Cell->Actors.Add(InActorHandle);
 		});
 
 		const int32 NewLevel = GetLevelForBox(Bounds);
@@ -177,18 +157,24 @@ void UWorldPartitionEditorSpatialHash::HashActor(FWorldPartitionHandle& InActorH
 
 		if (NewLevel > CurrentLevel)
 		{
-			ForEachIntersectingCells(Bounds, CurrentLevel, [&](const FCellCoord& CellCoord)
+			if (HashLevels.Num() <= NewLevel)
+			{
+				HashLevels.AddDefaulted(NewLevel - HashLevels.Num() + 1);
+			}
+
+			for (int32 Level = CurrentLevel; Level < NewLevel; Level++)
 			{
-				if (CurrentLevel ? HashNodes.Contains(CellCoord) : HashCells.Contains(CellCoord))
+				for (auto& HashLevelPair : HashLevels[Level])
 				{
-					FCellCoord LevelCellCoord = CellCoord;
+					FCellCoord LevelCellCoord = HashLevelPair.Key;
 					while (LevelCellCoord.Level < NewLevel)
 					{
 						const uint32 ChildIndex = LevelCellCoord.GetChildIndex();
 
 						LevelCellCoord = LevelCellCoord.GetParentCellCoord();
 
-						FCellNode& CellNode = HashNodes.FindOrAdd(LevelCellCoord);
+						FCellNodeElement& CellNodeElement = HashLevels[LevelCellCoord.Level].FindOrAdd(LevelCellCoord);
+						FCellNode& CellNode = CellNodeElement.Key;
 
 						// We can break updating when aggregated flags are already properly set for parent nodes
 						const bool bShouldBreak = CellNode.HasChildNodes();
@@ -205,7 +191,7 @@ void UWorldPartitionEditorSpatialHash::HashActor(FWorldPartitionHandle& InActorH
 						}
 					}
 				}
-			});
+			}
 		}
 	}
 }
@@ -214,54 +200,49 @@ void UWorldPartitionEditorSpatialHash::UnhashActor(FWorldPartitionHandle& InActo
 {
 	check(InActorHandle.IsValid());
 
-	if (IsActorAlwaysLoaded(InActorHandle))
+	if (!InActorHandle->GetIsSpatiallyLoaded())
 	{
-		AlwaysLoadedCell->RemoveActor(InActorHandle);
+		AlwaysLoadedCell->Actors.Remove(InActorHandle);
 	}
 	else
 	{
-		const FBox ActorBounds = GetActorBounds(InActorHandle);
+		const FBox ActorBounds = InActorHandle->GetBounds();
 		const int32 CurrentLevel = GetLevelForBox(Bounds);
+		const int32 ActorLevel = GetLevelForBox(ActorBounds);
 
-		ForEachIntersectingCells(ActorBounds, 0, [&](const FCellCoord& CellCoord)
+		ForEachIntersectingCells(ActorBounds, ActorLevel, [&](const FCellCoord& CellCoord)
 		{
-			UWorldPartitionEditorCell* EditorCell = HashCells.FindChecked(CellCoord);
+			FCellNodeElement& CellNodeElement = HashLevels[CellCoord.Level].FindChecked(CellCoord);
+			TUniquePtr<FCell>& Cell = CellNodeElement.Value;
+			check(Cell);
 
-			EditorCell->RemoveActor(InActorHandle);
+			Cell->Actors.Remove(InActorHandle);
 
-			if (!EditorCell->Actors.Num())
+			if (!Cell->Actors.Num())
 			{
-				verify(Cells.Remove(EditorCell));
-				verify(HashCells.Remove(CellCoord));
+				verify(Cells.Remove(Cell.Get()));
+				CellNodeElement.Value.Reset();
 
-				bool bClearChildMask = true;
-
-				FCellCoord LevelCellCoord = CellCoord;
-				while (LevelCellCoord.Level < CurrentLevel)
+				if (!CellNodeElement.Key.HasChildNodes())
 				{
-					const uint32 ChildIndex = LevelCellCoord.GetChildIndex();
-
-					LevelCellCoord = LevelCellCoord.GetParentCellCoord();
+					FCellCoord CurrentCellCoord = CellCoord;
+					while (CurrentCellCoord.Level < CurrentLevel)
+					{
+						FCellCoord ParentCellCoord = CurrentCellCoord.GetParentCellCoord();
+						FCellNodeElement& ParentCellNodeElement = HashLevels[ParentCellCoord.Level].FindChecked(ParentCellCoord);
+						FCellNode& ParentCellNode = ParentCellNodeElement.Key;
 
-					FCellNode& CellNode = HashNodes.FindChecked(LevelCellCoord);
+						const uint32 ChildIndex = CurrentCellCoord.GetChildIndex();
+						ParentCellNode.RemoveChildNode(ChildIndex);
 
-					if (bClearChildMask)
-					{
-						CellNode.RemoveChildNode(ChildIndex);
+						HashLevels[CurrentCellCoord.Level].Remove(CurrentCellCoord);
 
-						if (CellNode.HasChildNodes())
+						if (ParentCellNodeElement.Value || ParentCellNodeElement.Key.HasChildNodes())
 						{
-							bClearChildMask = false;
-						}
-						else
-						{
-							HashNodes.Remove(LevelCellCoord);
+							break;
 						}
-					}
 
-					if (!bClearChildMask)
-					{
-						break;
+						CurrentCellCoord = ParentCellCoord;
 					}
 				}
 
@@ -269,131 +250,103 @@ void UWorldPartitionEditorSpatialHash::UnhashActor(FWorldPartitionHandle& InActo
 			}
 		});
 	}
-
-	// Remove from all cells where it's back referenced
-	// This is important because the FWorldPartitionHandle will soon become invalid and Actors that were still 
-	// referencing this one will not be able to construct valid FWorldPartitionHandle to the actor we're 
-	// removing and it'll become impossible to clean the cells of their dangling actors
-	TArray<TTuple<UWorldPartitionEditorCell*, FGuid>> BackRefs;
-	BackReferences.MultiFind(InActorHandle->GetGuid(), BackRefs);
-
-	for (auto& BackRef : BackRefs)
-	{
-		BackRef.Get<0>()->RemoveActor(BackRef.Get<1>(), InActorHandle);
-	}
-
-	if (BackRefs.Num() != 0)
-	{
-		BackReferences.Remove(InActorHandle->GetGuid());
-	}
 }
 
-int32 UWorldPartitionEditorSpatialHash::ForEachIntersectingActor(const FBox& Box, TFunctionRef<void(FWorldPartitionActorDesc*)> InOperation)
+int32 UWorldPartitionEditorSpatialHash::ForEachIntersectingActor(const FBox& Box, TFunctionRef<void(FWorldPartitionActorDesc*)> InOperation, bool bIncludeSpatiallyLoadedActors, bool bIncludeNonSpatiallyLoadedActors)
 {
-	int32 NumIntersecting = 0;
+	TSet<FWorldPartitionActorDesc*> IntersectedActorDescs;
 
-	FWorldPartitionActorDesc::GlobalTag++;
-
-	ForEachIntersectingCell(Box, [&](UWorldPartitionEditorCell* EditorCell)
+	if (bIncludeSpatiallyLoadedActors)
 	{
-		for(UWorldPartitionEditorCell::FActorHandle& ActorDesc: EditorCell->Actors)
+		ForEachIntersectingCell(Box, [&](FCell* Cell)
 		{
-			if (ActorDesc.IsValid() && ActorDesc->Tag != FWorldPartitionActorDesc::GlobalTag)
+			for(FWorldPartitionHandle& ActorDesc: Cell->Actors)
 			{
-				if (Box.Intersect(ActorDesc->GetBounds()))
+				if (ActorDesc.IsValid())
 				{
-					InOperation(*ActorDesc);
-					NumIntersecting++;
-				}
+					bool bWasAlreadyInSet;
+					IntersectedActorDescs.Add(*ActorDesc, &bWasAlreadyInSet);
 
-				ActorDesc->Tag = FWorldPartitionActorDesc::GlobalTag;
+					if (!bWasAlreadyInSet)
+					{
+						if (Box.Intersect(ActorDesc->GetBounds()))
+						{
+							InOperation(*ActorDesc);
+						}
+					}
+				}
 			}
-		}
-	});
+		});
+	}
 
-	for(UWorldPartitionEditorCell::FActorHandle& ActorDesc: AlwaysLoadedCell->Actors)
+	if (bIncludeNonSpatiallyLoadedActors)
 	{
-		if (ActorDesc.IsValid() && ActorDesc->Tag != FWorldPartitionActorDesc::GlobalTag)
+		for(FWorldPartitionHandle& ActorDesc: AlwaysLoadedCell->Actors)
 		{
-			if (Box.Intersect(ActorDesc->GetBounds()))
+			if (ActorDesc.IsValid())
 			{
-				InOperation(*ActorDesc);
-				NumIntersecting++;
+				if (Box.Intersect(ActorDesc->GetBounds()))
+				{
+					bool bWasAlreadyInSet;
+					IntersectedActorDescs.Add(*ActorDesc, &bWasAlreadyInSet);
+				
+					if (!bWasAlreadyInSet)
+					{
+						InOperation(*ActorDesc);
+					}
+				}
 			}
-
-			ActorDesc->Tag = FWorldPartitionActorDesc::GlobalTag;
 		}
 	}
 
-	return NumIntersecting;
+	return IntersectedActorDescs.Num();
 }
 
-int32 UWorldPartitionEditorSpatialHash::ForEachIntersectingCellInner(const FBox& Box, const FCellCoord& CellCoord, TFunctionRef<void(UWorldPartitionEditorCell*)> InOperation)
+int32 UWorldPartitionEditorSpatialHash::ForEachIntersectingCellInner(const FBox& Box, const FCellCoord& CellCoord, TFunctionRef<void(FCell*)> InOperation)
 {
 	int32 NumIntersecting = 0;
 
-	if (CellCoord.Level)
+	if (const FCellNodeElement* CellNodeElement = HashLevels[CellCoord.Level].Find(CellCoord))
 	{
-		FCellNode CellNode;
-		if (FCellNode* CellNodePtr = HashNodes.Find(CellCoord))
+		if (CellNodeElement->Value)
 		{
-			CellNode = *CellNodePtr;
+			InOperation(CellNodeElement->Value.Get());
+			NumIntersecting++;
+		}
 
-			check(CellNode.HasChildNodes());
+		CellNodeElement->Key.ForEachChild([&](uint32 ChildIndex)
+		{
+			const FCellCoord ChildCellCoord = CellCoord.GetChildCellCoord(ChildIndex);
+			const FBox CellBounds = GetCellBounds(ChildCellCoord);
 
-			CellNode.ForEachChild([&](uint32 ChildIndex)
+			if (Box.Intersect(CellBounds))
 			{
-				const FCellCoord ChildCellCoord = CellCoord.GetChildCellCoord(ChildIndex);
-				const FBox CellBounds = GetCellBounds(ChildCellCoord);
-
-				if (Box.Intersect(CellBounds))
-				{
-					NumIntersecting += ForEachIntersectingCellInner(Box, ChildCellCoord, InOperation);
-				}
-			});
-		}
-	}
-	else
-	{
-		UWorldPartitionEditorCell* EditorCell = nullptr;
-		if (UWorldPartitionEditorCell** EditorCellPtr = HashCells.Find(CellCoord))
-		{
-			InOperation(*EditorCellPtr);
-			NumIntersecting++;
-		}
+				NumIntersecting += ForEachIntersectingCellInner(Box, ChildCellCoord, InOperation);
+			}
+		});
 	}
 
 	return NumIntersecting;
 }
 
-int32 UWorldPartitionEditorSpatialHash::ForEachIntersectingCell(const FBox& Box, TFunctionRef<void(UWorldPartitionEditorCell*)> InOperation)
+int32 UWorldPartitionEditorSpatialHash::ForEachIntersectingCell(const FBox& Box, TFunctionRef<void(FCell*)> InOperation)
 {
 	int32 NumIntersecting = 0;
 
-	const FBox SearchBox = Box.Overlap(Bounds);
-	const int32 SearchLevel = GetLevelForBox(SearchBox);
-
-	ForEachIntersectingCells(Box, SearchLevel, [&](const FCellCoord& CellCoord)
+	if (HashLevels.Num())
 	{
-		NumIntersecting += ForEachIntersectingCellInner(Box, CellCoord, InOperation);
-	});
+		const FBox SearchBox = Box.Overlap(Bounds);
 
-	return NumIntersecting;
-}
-
-int32 UWorldPartitionEditorSpatialHash::ForEachCell(TFunctionRef<void(UWorldPartitionEditorCell*)> InOperation)
-{
-	InOperation(AlwaysLoadedCell);
-	for (UWorldPartitionEditorCell* Cell: Cells)
-	{
-		InOperation(Cell);
+		if (SearchBox.IsValid)
+		{
+			ForEachIntersectingCells(SearchBox, HashLevels.Num() - 1, [&](const FCellCoord& CellCoord)
+			{
+				NumIntersecting += ForEachIntersectingCellInner(Box, CellCoord, InOperation);
+			});
+		}
 	}
-	return Cells.Num();
-}
 
-UWorldPartitionEditorCell* UWorldPartitionEditorSpatialHash::GetAlwaysLoadedCell()
-{
-	return AlwaysLoadedCell;
+	return NumIntersecting;
 }
 
 uint32 UWorldPartitionEditorSpatialHash::GetWantedEditorCellSize() const
@@ -417,17 +370,6 @@ void UWorldPartitionEditorSpatialHash::PostLoad()
 		WantedCellSize = 0;
 	}
 }
-
-
-void UWorldPartitionEditorSpatialHash::AddBackReference(const FGuid& Reference, UWorldPartitionEditorCell* Cell, const FGuid& Source)
-{
-	BackReferences.Add(Reference, MakeTuple(Cell, Source));
-}
-
-void UWorldPartitionEditorSpatialHash::RemoveBackReference(const FGuid& Reference, UWorldPartitionEditorCell* Cell, const FGuid& Source)
-{
-	BackReferences.RemoveSingle(Reference, MakeTuple(Cell, Source));
-}
 #endif
 
 #undef LOCTEXT_NAMESPACE
diff --git a/Engine/Source/Runtime/Engine/Private/WorldPartition/WorldPartitionVolume.cpp b/Engine/Source/Runtime/Engine/Private/WorldPartition/WorldPartitionVolume.cpp
--- a/Engine/Source/Runtime/Engine/Private/WorldPartition/WorldPartitionVolume.cpp
+++ b/Engine/Source/Runtime/Engine/Private/WorldPartition/WorldPartitionVolume.cpp
@@ -3,6 +3,7 @@
 #include "WorldPartition/WorldPartitionVolume.h"
 #include "WorldPartition/WorldPartitionEditorHash.h"
 #include "WorldPartition/WorldPartition.h"
+#include "WorldPartition/LoaderAdapter/LoaderAdapterActor.h"
 #include "Engine/World.h"
 #include "Misc/ScopedSlowTask.h"
 
@@ -14,23 +15,32 @@ AWorldPartitionVolume::AWorldPartitionVolume(const FObjectInitializer& ObjectIni
 #if WITH_EDITORONLY_DATA
 	bIsSpatiallyLoaded = false;
 #endif
-}
 
 #if WITH_EDITOR
-void AWorldPartitionVolume::LoadIntersectingCells(bool bIsFromUserChange)
-{
-	if (UWorldPartition* WorldPartition = GetWorld()->GetWorldPartition())
+	if (!IsTemplate() && GetWorld() && GetWorld()->GetWorldPartition())
 	{
-		WorldPartition->LoadEditorCells(GetStreamingBounds(), bIsFromUserChange);
+		WorldPartitionActorLoader = new FLoaderAdapterActor(this);
 	}
+#endif
 }
 
-void AWorldPartitionVolume::UnloadIntersectingCells(bool bIsFromUserChange)
+void AWorldPartitionVolume::BeginDestroy()
 {
-	if (UWorldPartition* WorldPartition = GetWorld()->GetWorldPartition())
+#if WITH_EDITOR
+	if (WorldPartitionActorLoader)
 	{
-		WorldPartition->UnloadEditorCells(GetStreamingBounds(), bIsFromUserChange);
+		delete WorldPartitionActorLoader;
+		WorldPartitionActorLoader = nullptr;
 	}
+#endif
+
+	Super::BeginDestroy();
+}
+
+#if WITH_EDITOR
+IWorldPartitionActorLoaderInterface::ILoaderAdapter* AWorldPartitionVolume::GetLoaderAdapter()
+{
+	return WorldPartitionActorLoader;
 }
 #endif
 
diff --git a/Engine/Source/Runtime/Engine/Private/WorldSettings.cpp b/Engine/Source/Runtime/Engine/Private/WorldSettings.cpp
--- a/Engine/Source/Runtime/Engine/Private/WorldSettings.cpp
+++ b/Engine/Source/Runtime/Engine/Private/WorldSettings.cpp
@@ -511,7 +511,7 @@ void AWorldSettings::SaveDefaultWorldPartitionSettings()
 
 	if (WorldPartition)
 	{
-		DefaultWorldPartitionSettings.LoadedEditorGridCells = WorldPartition->GetUserLoadedEditorGridCells();		
+		DefaultWorldPartitionSettings.LoadedEditorGridRegions = WorldPartition->GetUserLoadedEditorGridRegions();		
 
 		if (const UDataLayerSubsystem* DataLayerSubsystem = UWorld::GetSubsystem<UDataLayerSubsystem>(GetWorld()))
 		{
diff --git a/Engine/Source/Runtime/Engine/Public/LocationVolume.h b/Engine/Source/Runtime/Engine/Public/LocationVolume.h
new file mode 100644
--- /dev/null
+++ b/Engine/Source/Runtime/Engine/Public/LocationVolume.h
@@ -0,0 +1,50 @@
+// Copyright Epic Games, Inc. All Rights Reserved.
+
+#pragma once
+
+#include "CoreMinimal.h"
+#include "GameFramework/Volume.h"
+#include "WorldPartition/WorldPartitionActorLoaderInterface.h"
+
+#include "LocationVolume.generated.h"
+
+class FLoaderAdapterActor;
+
+/**
+ * A volume representing a location in the world
+ */
+UCLASS()
+class ENGINE_API ALocationVolume : public AVolume, public IWorldPartitionActorLoaderInterface
+{
+	GENERATED_UCLASS_BODY()
+
+public:
+	//~ Begin UObject Interface
+	virtual void BeginDestroy() override;
+	//~ End UObject Interface
+
+	//~ Begin AActor Interface
+	virtual bool IsEditorOnly() const override { return !bIsRuntime; }
+#if WITH_EDITOR
+	virtual bool CanChangeIsSpatiallyLoadedFlag() const override { return false; }
+	virtual bool SupportsDataLayer() const override { return false; }
+#endif
+	//~ End AActor Interface
+
+#if WITH_EDITOR
+	//~ Begin IWorldPartitionActorLoaderInterface interface
+	virtual ILoaderAdapter* GetLoaderAdapter() override;
+	//~ End IWorldPartitionActorLoaderInterface interface
+#endif
+
+	UPROPERTY(EditAnywhere, Category=LocationVolume)
+	FColor DebugColor;
+
+	UPROPERTY(EditAnywhere, Category=LocationVolume)
+	uint8 bIsRuntime : 1;
+
+#if WITH_EDITOR
+private:
+	FLoaderAdapterActor* WorldPartitionActorLoader;
+#endif
+};
\ No newline at end of file
diff --git a/Engine/Source/Runtime/Engine/Public/WorldPartition/HLOD/HLODActorDesc.h b/Engine/Source/Runtime/Engine/Public/WorldPartition/HLOD/HLODActorDesc.h
--- a/Engine/Source/Runtime/Engine/Public/WorldPartition/HLOD/HLODActorDesc.h
+++ b/Engine/Source/Runtime/Engine/Public/WorldPartition/HLOD/HLODActorDesc.h
@@ -26,7 +26,7 @@ public:
 
 	static uint64 ComputeCellHash(const FString HLODLayerName, uint64 GridIndexX, uint64 GridIndexY, uint64 GridIndexZ, FDataLayersID DataLayersID);
 
-	virtual bool ShouldBeLoadedByEditorCells() const { return false; }
+	virtual bool ShouldBeLoadedByEditor() const { return false; }
 
 protected:
 	virtual void Init(const AActor* InActor) override;
diff --git a/Engine/Source/Runtime/Engine/Public/WorldPartition/LoaderAdapter/LoaderAdapterActor.h b/Engine/Source/Runtime/Engine/Public/WorldPartition/LoaderAdapter/LoaderAdapterActor.h
new file mode 100644
--- /dev/null
+++ b/Engine/Source/Runtime/Engine/Public/WorldPartition/LoaderAdapter/LoaderAdapterActor.h
@@ -0,0 +1,26 @@
+// Copyright Epic Games, Inc. All Rights Reserved.
+#pragma once
+
+#include "CoreMinimal.h"
+#include "WorldPartition/WorldPartitionActorLoaderInterface.h"
+
+#if WITH_EDITOR
+/** Actor loader adapter that will do a spatial query based on the actor's brush  */
+class ENGINE_API FLoaderAdapterActor : public IWorldPartitionActorLoaderInterface::ILoaderAdapterSpatial
+{
+public:
+	FLoaderAdapterActor(AActor* InActor);
+
+	//~ Begin IWorldPartitionActorLoaderInterface::ILoader interface
+	virtual TOptional<FBox> GetBoundingBox() const override;
+	virtual TOptional<FString> GetLabel() const override;
+	//~ End IWorldPartitionActorLoaderInterface::ILoader interface
+
+	//~ Begin IWorldPartitionActorLoaderInterface::ILoaderAdapterSpatial interface
+	virtual bool Intersect(const FBox& Box) const override;
+	//~ End IWorldPartitionActorLoaderInterface::ILoaderAdapterSpatial interface
+
+private:
+	AActor* Actor;
+};
+#endif
diff --git a/Engine/Source/Runtime/Engine/Public/WorldPartition/LoaderAdapter/LoaderAdapterShape.h b/Engine/Source/Runtime/Engine/Public/WorldPartition/LoaderAdapter/LoaderAdapterShape.h
new file mode 100644
--- /dev/null
+++ b/Engine/Source/Runtime/Engine/Public/WorldPartition/LoaderAdapter/LoaderAdapterShape.h
@@ -0,0 +1,27 @@
+// Copyright Epic Games, Inc. All Rights Reserved.
+#pragma once
+
+#include "CoreMinimal.h"
+#include "WorldPartition/WorldPartitionActorLoaderInterface.h"
+
+#if WITH_EDITOR
+/** Actor loader adapter that will do a spatial query based on a shape */
+class ENGINE_API FLoaderAdapterShape : public IWorldPartitionActorLoaderInterface::ILoaderAdapterSpatial
+{
+public:
+	FLoaderAdapterShape(UWorld* InWorld, const FBox& InBoundingBox, const FString& InLabel);
+
+	//~ Begin IWorldPartitionActorLoaderInterface::ILoader interface
+	virtual TOptional<FBox> GetBoundingBox() const override;
+	virtual TOptional<FString> GetLabel() const override;
+	//~ End IWorldPartitionActorLoaderInterface::ILoader interface
+
+	//~ Begin IWorldPartitionActorLoaderInterface::ILoaderAdapterSpatial interface
+	virtual bool Intersect(const FBox& Box) const override;
+	//~ End IWorldPartitionActorLoaderInterface::ILoaderAdapterSpatial interface
+
+private:
+	FBox BoundingBox;
+	FString Label;
+};
+#endif
\ No newline at end of file
diff --git a/Engine/Source/Runtime/Engine/Public/WorldPartition/WorldPartition.h b/Engine/Source/Runtime/Engine/Public/WorldPartition/WorldPartition.h
--- a/Engine/Source/Runtime/Engine/Public/WorldPartition/WorldPartition.h
+++ b/Engine/Source/Runtime/Engine/Public/WorldPartition/WorldPartition.h
@@ -15,6 +15,7 @@
 #include "WorldPartition/ActorDescContainer.h"
 
 #if WITH_EDITOR
+#include "WorldPartition/WorldPartitionActorLoaderInterface.h"
 #include "PackageSourceControlHelper.h"
 #include "CookPackageSplitter.h"
 #endif
@@ -22,12 +23,12 @@
 #include "WorldPartition.generated.h"
 
 class FWorldPartitionActorDesc;
-class UWorldPartitionEditorCell;
 class UWorldPartitionEditorHash;
 class UWorldPartitionRuntimeCell;
 class UWorldPartitionRuntimeHash;
 class UWorldPartitionStreamingPolicy;
 class IStreamingGenerationErrorHandler;
+class FLoaderAdapterAlwaysLoadedActors;
 class FHLODActorDesc;
 class UCanvas;
 
@@ -76,7 +77,6 @@ class ENGINE_API UWorldPartition final : public UActorDescContainer
 
 	friend class FWorldPartitionActorDesc;
 	friend class FWorldPartitionConverter;
-	friend class UWorldPartitionEditorCell;
 	friend class UWorldPartitionConvertCommandlet;
 	friend class FWorldPartitionEditorModule;
 	friend class FWorldPartitionDetails;
@@ -85,20 +85,18 @@ class ENGINE_API UWorldPartition final : public UActorDescContainer
 public:
 #if WITH_EDITOR
 	static UWorldPartition* CreateOrRepairWorldPartition(AWorldSettings* WorldSettings, TSubclassOf<UWorldPartitionEditorHash> EditorHashClass = nullptr, TSubclassOf<UWorldPartitionRuntimeHash> RuntimeHashClass = nullptr);
-	
-	DECLARE_MULTICAST_DELEGATE_OneParam(FCancelWorldPartitionUpdateEditorCellsDelegate, UWorldPartition*);
-	FCancelWorldPartitionUpdateEditorCellsDelegate OnCancelWorldPartitionUpdateEditorCells;
 #endif
 
 	DECLARE_MULTICAST_DELEGATE_OneParam(FWorldPartitionInitializeDelegate, UWorldPartition*);
 	
 	UE_DEPRECATED(5.1, "Please use FWorldPartitionInitializedEvent& UWorld::OnWorldPartitionInitialized() instead.")
 	FWorldPartitionInitializeDelegate OnWorldPartitionInitialized;
+
 	UE_DEPRECATED(5.1, "Please use FWorldPartitionInitializedEvent& UWorld::OnWorldPartitionUninitialized() instead.")
 	FWorldPartitionInitializeDelegate OnWorldPartitionUninitialized;
 
 #if WITH_EDITOR
-	TArray<FName> GetUserLoadedEditorGridCells() const;
+	TArray<FBox> GetUserLoadedEditorGridRegions() const;
 
 public:
 	bool SupportsStreaming() const;
@@ -143,11 +141,6 @@ public:
 #if WITH_EDITOR
 	FName GetWorldPartitionEditorName() const;
 
-	void LoadEditorCells(const FBox& Box, bool bIsFromUserChange);
-	void LoadEditorCells(const TArray<FName>& CellNames, bool bIsFromUserChange);
-	void UnloadEditorCells(const FBox& Box, bool bIsFromUserChange);
-	bool AreEditorCellsLoaded(const FBox& Box);
-
 	// PIE/Game/Cook Methods
 	bool GenerateStreaming(TArray<FString>* OutPackagesToGenerate = nullptr);
 	void RemapSoftObjectPath(FSoftObjectPath& ObjectPath);
@@ -176,6 +169,8 @@ public:
 	AActor* PinActor(const FGuid& ActorGuid);
 	void UnpinActor(const FGuid& ActorGuid);
 	bool IsActorPinned(const FGuid& ActorGuid) const;
+
+	void LoadLastLoadedRegions();
 #endif
 
 public:
@@ -210,6 +205,8 @@ public:
 	UPROPERTY(DuplicateTransient)
 	TObjectPtr<UWorldPartitionEditorHash> EditorHash;
 
+	FLoaderAdapterAlwaysLoadedActors* AlwaysLoadedActors;
+
 	IWorldPartitionEditor* WorldPartitionEditor;
 
 private:
@@ -280,16 +277,35 @@ private:
 	void UnregisterDelegates();	
 
 #if WITH_EDITOR
-	void OnActorDataLayersEditorLoadingStateChanged(bool bFromUserOperation);
-
-	void UpdateLoadingEditorCell(UWorldPartitionEditorCell* Cell, bool bShouldBeLoaded, bool bFromUserOperation);
 	void HashActorDesc(FWorldPartitionActorDesc* ActorDesc);
 	void UnhashActorDesc(FWorldPartitionActorDesc* ActorDesc);
-	bool ShouldActorBeLoadedByEditorCells(const FWorldPartitionActorDesc* ActorDesc) const;
-	bool UpdateEditorCells(TFunctionRef<bool(TArray<UWorldPartitionEditorCell*>&)> GetCellsToProcess, bool bIsCellShouldBeLoaded, bool bIsFromUserChange);
 
 	TMap<FGuid, FWorldPartitionReference> PinnedActors;
 	TMap<FGuid, TMap<FGuid, FWorldPartitionReference>> PinnedActorRefs;
+
+public:
+	// Editor loader adapters management
+	template <typename T, typename... ArgsType>
+	T* CreateEditorLoaderAdapter(ArgsType&&... Args)
+	{
+		T* LoaderAdapter = new T(Forward<ArgsType>(Args)...);
+		RegisteredEditorLoaderAdapters.Add(LoaderAdapter);
+		return LoaderAdapter;
+	}
+
+	void ReleaseEditorLoaderAdapter(IWorldPartitionActorLoaderInterface::ILoaderAdapter* LoaderAdapter)
+	{
+		verify(RegisteredEditorLoaderAdapters.Remove(LoaderAdapter) != INDEX_NONE);
+		delete LoaderAdapter;
+	}
+
+	const TSet<IWorldPartitionActorLoaderInterface::ILoaderAdapter*>& GetRegisteredEditorLoaderAdapters() const
+	{
+		return RegisteredEditorLoaderAdapters;
+	}
+
+private:
+	TSet<IWorldPartitionActorLoaderInterface::ILoaderAdapter*> RegisteredEditorLoaderAdapters;
 #endif
 
 #if !UE_BUILD_SHIPPING
diff --git a/Engine/Source/Runtime/Engine/Public/WorldPartition/WorldPartitionActorDesc.h b/Engine/Source/Runtime/Engine/Public/WorldPartition/WorldPartitionActorDesc.h
--- a/Engine/Source/Runtime/Engine/Public/WorldPartition/WorldPartitionActorDesc.h
+++ b/Engine/Source/Runtime/Engine/Public/WorldPartition/WorldPartitionActorDesc.h
@@ -147,7 +147,11 @@ public:
 	AActor* GetActor(bool bEvenIfPendingKill=true, bool bEvenIfUnreachable=false) const;
 	AActor* Load() const;
 	virtual void Unload();
-	virtual bool ShouldBeLoadedByEditorCells() const { return true; }
+
+	UE_DEPRECATED(5.1, "ShouldBeLoadedByEditorCells is deprecated, ShouldBeLoadedByEditor should be used instead.")
+	bool ShouldBeLoadedByEditorCells() const { return ShouldBeLoadedByEditor(); }
+
+	virtual bool ShouldBeLoadedByEditor() const { return true; }
 
 	void RegisterActor();
 	void UnregisterActor();
@@ -212,10 +216,5 @@ protected:
 	UActorDescContainer*			Container;
 	TArray<FName>					DataLayerInstanceNames;
 	bool							bIsForcedNonSpatiallyLoaded;
-
-public:
-	// Tagging
-	mutable uint32					Tag;
-	static uint32					GlobalTag;
 #endif
 };
diff --git a/Engine/Source/Runtime/Engine/Public/WorldPartition/WorldPartitionActorDescView.h b/Engine/Source/Runtime/Engine/Public/WorldPartition/WorldPartitionActorDescView.h
--- a/Engine/Source/Runtime/Engine/Public/WorldPartition/WorldPartitionActorDescView.h
+++ b/Engine/Source/Runtime/Engine/Public/WorldPartition/WorldPartitionActorDescView.h
@@ -42,7 +42,6 @@ public:
 	FBox GetBounds() const;
 	const TArray<FGuid>& GetReferences() const;
 	FString ToString() const;
-	uint32 GetTag() const;
 	const FGuid& GetParentActor() const;
 	FName GetActorName() const;
 	const FGuid& GetFolderGuid() const;
diff --git a/Engine/Source/Runtime/Engine/Public/WorldPartition/WorldPartitionActorLoaderInterface.h b/Engine/Source/Runtime/Engine/Public/WorldPartition/WorldPartitionActorLoaderInterface.h
new file mode 100644
--- /dev/null
+++ b/Engine/Source/Runtime/Engine/Public/WorldPartition/WorldPartitionActorLoaderInterface.h
@@ -0,0 +1,96 @@
+// Copyright Epic Games, Inc. All Rights Reserved.
+
+#pragma once
+
+#include "CoreMinimal.h"
+#include "UObject/Interface.h"
+#include "Misc/Optional.h"
+#include "WorldPartition/WorldPartitionHandle.h"
+#include "WorldPartitionActorLoaderInterface.generated.h"
+
+UINTERFACE()
+class ENGINE_API UWorldPartitionActorLoaderInterface : public UInterface
+{
+	GENERATED_UINTERFACE_BODY()
+};
+
+class ENGINE_API IWorldPartitionActorLoaderInterface
+{
+	GENERATED_IINTERFACE_BODY()
+
+#if WITH_EDITOR
+public:
+	/** Base class for actor loaders */
+	class ENGINE_API ILoaderAdapter
+	{
+	public:
+		ILoaderAdapter(UWorld* InWorld);
+		virtual ~ILoaderAdapter();
+
+		bool Load();
+		bool Unload();
+		bool IsLoaded() const;
+
+		bool GetUserCreated() const { return bUserCreated; }
+		void SetUserCreated(bool bValue) { bUserCreated = bValue; }
+
+		// Interface
+		virtual TOptional<FBox> GetBoundingBox() const { return TOptional<FBox>(); }
+		virtual TOptional<FString> GetLabel() const { return TOptional<FString>(); }
+
+	protected:
+		virtual bool Intersect(const FBox& Box) const =0;
+		virtual void ForEachActor(TFunctionRef<void(const FWorldPartitionHandle&)> InOperation) const =0;
+
+		bool RefreshLoadedState();
+		bool ShouldActorBeLoaded(const FWorldPartitionHandle& Actor) const;
+		void PostLoadedStateChanged(bool bUnloadedActors);
+		bool AllowUnloadingActors(const TArray<FWorldPartitionHandle>& ActorsToUnload) const;
+		void AddReferenceToActor(FWorldPartitionHandle& Actor);
+		void OnActorDataLayersEditorLoadingStateChanged(bool bFromUserOperation);		
+
+		UWorld* World;
+
+		uint8 bLoaded : 1;
+		uint8 bUserCreated : 1;
+
+		TMap<FGuid, TMap<FGuid, FWorldPartitionReference>> ActorReferences;
+	};
+
+	/** Base class for actor loaders that contains a specific list of actors */
+	class ENGINE_API FLoaderAdapterList : public ILoaderAdapter
+	{
+	public:
+		FLoaderAdapterList(UWorld* InWorld);
+		virtual ~FLoaderAdapterList() {}
+
+		//~ Begin ILoaderAdapter interface
+		virtual bool Intersect(const FBox& Box) const { return true; }
+		virtual void ForEachActor(TFunctionRef<void(const FWorldPartitionHandle&)> InOperation) const;
+		//~ End ILoaderAdapter interface
+
+	protected:
+		TSet<FWorldPartitionHandle> Actors;
+	};
+
+	/** Base class for actor loaders that requires spatial queries */
+	class ENGINE_API ILoaderAdapterSpatial : public ILoaderAdapter
+	{
+	public:
+		ILoaderAdapterSpatial(UWorld* InWorld);
+		virtual ~ILoaderAdapterSpatial() {}
+
+		//~ Begin ILoaderAdapter interface
+		virtual void ForEachActor(TFunctionRef<void(const FWorldPartitionHandle&)> InOperation) const;
+		//~ End ILoaderAdapter interface
+
+		// Interface
+		virtual bool Intersect(const FBox& Box) const =0;
+
+		bool bIncludeSpatiallyLoadedActors;
+		bool bIncludeNonSpatiallyLoadedActors;
+	};
+
+	virtual ILoaderAdapter* GetLoaderAdapter() =0;
+#endif
+};
\ No newline at end of file
diff --git a/Engine/Source/Runtime/Engine/Public/WorldPartition/WorldPartitionEditorCell.h b/Engine/Source/Runtime/Engine/Public/WorldPartition/WorldPartitionEditorCell.h
deleted file mode 100644
diff --git a/Engine/Source/Runtime/Engine/Public/WorldPartition/WorldPartitionEditorHash.h b/Engine/Source/Runtime/Engine/Public/WorldPartition/WorldPartitionEditorHash.h
--- a/Engine/Source/Runtime/Engine/Public/WorldPartition/WorldPartitionEditorHash.h
+++ b/Engine/Source/Runtime/Engine/Public/WorldPartition/WorldPartitionEditorHash.h
@@ -8,8 +8,6 @@
 #include "WorldPartition/WorldPartitionHandle.h"
 #include "WorldPartitionEditorHash.generated.h"
 
-class UWorldPartitionEditorCell;
-
 UCLASS(Abstract, Config=Engine, Within = WorldPartition)
 class ENGINE_API UWorldPartitionEditorHash : public UObject
 {
@@ -26,40 +24,9 @@ public:
 	virtual void HashActor(FWorldPartitionHandle& InActorHandle) PURE_VIRTUAL(UWorldPartitionEditorHash::HashActor, ;);
 	virtual void UnhashActor(FWorldPartitionHandle& InActorHandle) PURE_VIRTUAL(UWorldPartitionEditorHash::UnhashActor, ;);
 
-	virtual int32 ForEachIntersectingActor(const FBox& Box, TFunctionRef<void(FWorldPartitionActorDesc*)> InOperation) PURE_VIRTUAL(UWorldPartitionEditorHash::ForEachIntersectingActor, return 0;);
-	virtual int32 ForEachIntersectingCell(const FBox& Box, TFunctionRef<void(UWorldPartitionEditorCell*)> InOperation) PURE_VIRTUAL(UWorldPartitionEditorHash::ForEachIntersectingCell, return 0;);
-	virtual int32 ForEachCell(TFunctionRef<void(UWorldPartitionEditorCell*)> InOperation) PURE_VIRTUAL(UWorldPartitionEditorHash::ForEachCell, return 0;);
-	virtual UWorldPartitionEditorCell* GetAlwaysLoadedCell() PURE_VIRTUAL(UWorldPartitionEditorHash::GetAlwaysLoadedCell, return nullptr;);
+	virtual int32 ForEachIntersectingActor(const FBox& Box, TFunctionRef<void(FWorldPartitionActorDesc*)> InOperation, bool bIncludeSpatiallyLoadedActors = true, bool bIncludeNonSpatiallyLoadedActors = true) PURE_VIRTUAL(UWorldPartitionEditorHash::ForEachIntersectingActor, return 0;);
 
 	virtual uint32 GetWantedEditorCellSize() const PURE_VIRTUAL(UWorldPartitionEditorHash::GetWantedEditorCellSize, return 0;);
 	virtual void SetEditorWantedCellSize(uint32 InCellSize) PURE_VIRTUAL(UWorldPartitionEditorHash::SetEditorWantedCellSize, );
-
-	virtual void AddBackReference(const FGuid& Reference, UWorldPartitionEditorCell* Cell, const FGuid& Source) PURE_VIRTUAL(UWorldPartitionEditorHash::AddBackReference, ;);
-	virtual void RemoveBackReference(const FGuid& Reference, UWorldPartitionEditorCell* Cell, const FGuid& Source) PURE_VIRTUAL(UWorldPartitionEditorHash::RemoveBackReference, ;);
-
-	// Helpers
-	inline int32 GetIntersectingActors(const FBox& Box, TArray<FWorldPartitionActorDesc*>& OutActors)
-	{
-		return ForEachIntersectingActor(Box, [&OutActors](FWorldPartitionActorDesc* ActorDesc)
-		{
-			OutActors.Add(ActorDesc);
-		});
-	}
-
-	int32 GetIntersectingCells(const FBox& Box, TArray<UWorldPartitionEditorCell*>& OutCells)
-	{
-		return ForEachIntersectingCell(Box, [&OutCells](UWorldPartitionEditorCell* Cell)
-		{
-			OutCells.Add(Cell);
-		});
-	}
-
-	int32 GetAllCells(TArray<UWorldPartitionEditorCell*>& OutCells)
-	{
-		return ForEachCell([&OutCells](UWorldPartitionEditorCell* Cell)
-		{
-			OutCells.Add(Cell);
-		});
-	}
 #endif
 };
diff --git a/Engine/Source/Runtime/Engine/Public/WorldPartition/WorldPartitionEditorPerProjectUserSettings.h b/Engine/Source/Runtime/Engine/Public/WorldPartition/WorldPartitionEditorPerProjectUserSettings.h
--- a/Engine/Source/Runtime/Engine/Public/WorldPartition/WorldPartitionEditorPerProjectUserSettings.h
+++ b/Engine/Source/Runtime/Engine/Public/WorldPartition/WorldPartitionEditorPerProjectUserSettings.h
@@ -18,13 +18,13 @@ struct FWorldPartitionPerWorldSettings
 	FWorldPartitionPerWorldSettings()
 	{}
 
-	FWorldPartitionPerWorldSettings(TArray<FName>& InLoadedEditorGridCells)
-		: LoadedEditorGridCells(InLoadedEditorGridCells)
+	FWorldPartitionPerWorldSettings(TArray<FBox>& InLoadedEditorGridRegions)
+		: LoadedEditorGridRegions(InLoadedEditorGridRegions)
 	{}
 
 	void Reset()
 	{
-		LoadedEditorGridCells.Empty();
+		LoadedEditorGridRegions.Empty();
 		NotLoadedDataLayers.Empty();
 		LoadedDataLayers.Empty();
 		EditorGridConfigHash = 0;
@@ -33,7 +33,7 @@ struct FWorldPartitionPerWorldSettings
 
 #if WITH_EDITORONLY_DATA
 	UPROPERTY()
-	TArray<FName> LoadedEditorGridCells;
+	TArray<FBox> LoadedEditorGridRegions;
 
 	UPROPERTY()
 	TArray<FName> NotLoadedDataLayers;
@@ -62,31 +62,31 @@ public:
 		, bShowOnlySelectedActors(false)
 		, bHighlightSelectedDataLayers(true)
 		, bHideLevelInstanceContent(true)
-		, bDisableLoadingOfLastLoadedCells(false)
+		, bDisableLoadingOfLastLoadedRegions(false)
 #endif
 	{}
 
 #if WITH_EDITOR
 	void UpdateEditorGridConfigHash(UWorld* InWorld);
 	static void UpdateEditorGridConfigHash(UWorld* InWorld, FWorldPartitionPerWorldSettings& PerWorldSettings);
-	TArray<FName> GetEditorGridLoadedCells(UWorld* InWorld) const;
-	void SetEditorGridLoadedCells(UWorld* InWorld, const TArray<FName>& InEditorGridLoadedCells);
+	TArray<FBox> GetEditorGridLoadedRegions(UWorld* InWorld) const;
+	void SetEditorGridLoadedRegions(UWorld* InWorld, const TArray<FBox>& InEditorGridLoadedRegions);
 
-	bool GetEnableLoadingOfLastLoadedCells() const
+	bool GetEnableLoadingOfLastLoadedRegions() const
 	{
-		return !bDisableLoadingOfLastLoadedCells;
+		return !bDisableLoadingOfLastLoadedRegions;
 	}
 
-	bool GetBugItGoLoadCells() const
+	bool GetBugItGoLoadRegion() const
 	{
-		return bBugItGoLoadCells;
+		return bBugItGoLoadRegion;
 	}
 
-	void SetBugItGoLoadCells(bool bInBugItGoLoadCells)
+	void SetBugItGoLoadRegion(bool bInBugItGoLoadRegion)
 	{
-		if (bBugItGoLoadCells != bInBugItGoLoadCells)
+		if (bBugItGoLoadRegion != bInBugItGoLoadRegion)
 		{
-			bBugItGoLoadCells = bInBugItGoLoadCells;
+			bBugItGoLoadRegion = bInBugItGoLoadRegion;
 			SaveConfig();
 		}
 	}
@@ -151,10 +151,10 @@ private:
 	}
 
 	UPROPERTY(config)
-	uint32 bDisableLoadingOfLastLoadedCells : 1;
+	uint32 bDisableLoadingOfLastLoadedRegions : 1;
 
 	UPROPERTY(config)
-	uint32 bBugItGoLoadCells : 1;
+	uint32 bBugItGoLoadRegion : 1;
 
 	UPROPERTY(config)
 	uint32 bShowCellCoords : 1;
diff --git a/Engine/Source/Runtime/Engine/Public/WorldPartition/WorldPartitionEditorSpatialHash.h b/Engine/Source/Runtime/Engine/Public/WorldPartition/WorldPartitionEditorSpatialHash.h
--- a/Engine/Source/Runtime/Engine/Public/WorldPartition/WorldPartitionEditorSpatialHash.h
+++ b/Engine/Source/Runtime/Engine/Public/WorldPartition/WorldPartitionEditorSpatialHash.h
@@ -17,6 +17,17 @@ class ENGINE_API UWorldPartitionEditorSpatialHash : public UWorldPartitionEditor
 #if WITH_EDITOR
 	friend class SWorldPartitionEditorGridSpatialHash;
 
+	struct FCell
+	{
+	public:
+		FCell()
+			: Bounds(ForceInitToZero)
+		{}
+
+		FBox Bounds;
+		TSet<FWorldPartitionHandle> Actors;
+	};
+
 	struct FCellCoord
 	{
 		FCellCoord(int64 InX, int64 InY, int64 InZ, int32 InLevel)
@@ -189,25 +200,17 @@ public:
 	virtual void HashActor(FWorldPartitionHandle& InActorHandle) override;
 	virtual void UnhashActor(FWorldPartitionHandle& InActorHandle) override;
 
-	virtual int32 ForEachIntersectingActor(const FBox& Box, TFunctionRef<void(FWorldPartitionActorDesc*)> InOperation) override;
-	virtual int32 ForEachIntersectingCell(const FBox& Box, TFunctionRef<void(UWorldPartitionEditorCell*)> InOperation) override;
-	virtual int32 ForEachCell(TFunctionRef<void(UWorldPartitionEditorCell*)> InOperation) override;
-	virtual UWorldPartitionEditorCell* GetAlwaysLoadedCell() override;
+	virtual int32 ForEachIntersectingActor(const FBox& Box, TFunctionRef<void(FWorldPartitionActorDesc*)> InOperation, bool bIncludeSpatiallyLoadedActors = true, bool bIncludeNonSpatiallyLoadedActors = true) override;
 
 	virtual uint32 GetWantedEditorCellSize() const override;
 	virtual void SetEditorWantedCellSize(uint32 InCellSize) override;
-
-	virtual void AddBackReference(const FGuid& Reference, UWorldPartitionEditorCell* Cell, const FGuid& Source) override;
-	virtual void RemoveBackReference(const FGuid& Reference, UWorldPartitionEditorCell* Cell, const FGuid& Source) override;
-		
 	// UWorldPartitionEditorHash interface end
 #endif
 
 #if WITH_EDITORONLY_DATA
 private:
-	FBox GetActorBounds(const FWorldPartitionHandle& InActorHandle) const;
-	bool IsActorAlwaysLoaded(const FWorldPartitionHandle& InActorHandle) const;
-	int32 ForEachIntersectingCellInner(const FBox& Box, const FCellCoord& CellCoord, TFunctionRef<void(UWorldPartitionEditorCell*)> InOperation);
+	int32 ForEachIntersectingCell(const FBox& Box, TFunctionRef<void(FCell*)> InOperation);
+	int32 ForEachIntersectingCellInner(const FBox& Box, const FCellCoord& CellCoord, TFunctionRef<void(FCell*)> InOperation);
 
 	UPROPERTY(Config)
 	int32 CellSize;
@@ -215,19 +218,16 @@ private:
 	UPROPERTY()
 	int32 WantedCellSize;
 
-	TMap<FCellCoord, FCellNode> HashNodes;
-	TMap<FCellCoord, UWorldPartitionEditorCell*> HashCells;
-	
-	TMultiMap<FGuid, TTuple<UWorldPartitionEditorCell*, FGuid>>	BackReferences;
+	// Dynamic sparse octree structure
+	typedef TTuple<FCellNode, TUniquePtr<FCell>> FCellNodeElement;
+	typedef TMap<FCellCoord, FCellNodeElement> FCellNodeHashLevel;
+	TArray<FCellNodeHashLevel, TInlineAllocator<20>> HashLevels;
 
-	UPROPERTY(Transient)
-	TSet<TObjectPtr<UWorldPartitionEditorCell>> Cells;
+	TSet<FCell*> Cells;
+	TUniquePtr<FCell> AlwaysLoadedCell;
 	
 	FBox Bounds;
-	bool bBoundsDirty;
-	
-	UPROPERTY(Transient)
-	TObjectPtr<UWorldPartitionEditorCell> AlwaysLoadedCell;
+	bool bBoundsDirty;	
 
 public:
 	UPROPERTY(Config, meta = (AllowedClasses = "Texture2D, MaterialInterface"))
diff --git a/Engine/Source/Runtime/Engine/Public/WorldPartition/WorldPartitionVolume.h b/Engine/Source/Runtime/Engine/Public/WorldPartition/WorldPartitionVolume.h
--- a/Engine/Source/Runtime/Engine/Public/WorldPartition/WorldPartitionVolume.h
+++ b/Engine/Source/Runtime/Engine/Public/WorldPartition/WorldPartitionVolume.h
@@ -3,29 +3,42 @@
 
 #include "CoreMinimal.h"
 #include "GameFramework/Volume.h"
+#include "WorldPartition/WorldPartitionActorLoaderInterface.h"
 #include "WorldPartitionVolume.generated.h"
 
+class FLoaderAdapterActor;
+
 /**
  * A world partition volume to allow loading cells inside (editor-only)
  */
 UCLASS()
-class ENGINE_API AWorldPartitionVolume : public AVolume
+class ENGINE_API AWorldPartitionVolume : public AVolume, public IWorldPartitionActorLoaderInterface
 {
 	GENERATED_BODY()
 
 public:
 	AWorldPartitionVolume(const FObjectInitializer& ObjectInitializer = FObjectInitializer::Get());
 
-	// Begin AActor
-	virtual bool IsEditorOnly() const final { return true; }
-	// End AActor
+	//~ Begin UObject Interface
+	virtual void BeginDestroy() override;
+	//~ End UObject Interface
 
+	//~ Begin AActor Interface
+	virtual bool IsEditorOnly() const override { return true; }
 #if WITH_EDITOR
-	void LoadIntersectingCells(bool bIsFromUserChange);
-	void UnloadIntersectingCells(bool bIsFromUserChange);
 	virtual bool CanChangeIsSpatiallyLoadedFlag() const override { return false; }
+	virtual bool SupportsDataLayer() const override { return false; }
+#endif
+	//~ End AActor Interface
+
+#if WITH_EDITOR
+	//~ Begin IWorldPartitionActorLoaderInterface interface
+	virtual ILoaderAdapter* GetLoaderAdapter() override;
+	//~ End IWorldPartitionActorLoaderInterface interface
+#endif
 
+#if WITH_EDITOR
 private:
-	virtual bool SupportsDataLayer() const override { return false; }
+	FLoaderAdapterActor* WorldPartitionActorLoader;
 #endif
 };
\ No newline at end of file
diff --git a/Engine/Source/Runtime/Foliage/Private/ProceduralFoliageComponent.cpp b/Engine/Source/Runtime/Foliage/Private/ProceduralFoliageComponent.cpp
--- a/Engine/Source/Runtime/Foliage/Private/ProceduralFoliageComponent.cpp
+++ b/Engine/Source/Runtime/Foliage/Private/ProceduralFoliageComponent.cpp
@@ -16,9 +16,9 @@
 #include "WorldPartition/WorldPartition.h"
 #include "WorldPartition/WorldPartitionSubsystem.h"
 #include "WorldPartition/DataLayer/DataLayerAsset.h"
+#include "WorldPartition/WorldPartitionActorLoaderInterface.h"
 #endif
 
-
 #define LOCTEXT_NAMESPACE "ProceduralFoliage"
 
 UProceduralFoliageComponent::UProceduralFoliageComponent(const FObjectInitializer& ObjectInitializer)
@@ -101,6 +101,36 @@ void UProceduralFoliageComponent::GetTileLayout(FTileLayout& OutTileLayout) cons
 	}
 }
 
+#if WITH_EDITOR
+void UProceduralFoliageComponent::LoadSimulatedRegion()
+{
+	if (GetOwner()->Implements<UWorldPartitionActorLoaderInterface>())
+	{
+		IWorldPartitionActorLoaderInterface::ILoaderAdapter* LoaderAdapter = Cast<IWorldPartitionActorLoaderInterface>(GetOwner())->GetLoaderAdapter();
+		LoaderAdapter->Load();
+	}
+}
+
+void UProceduralFoliageComponent::UnloadSimulatedRegion()
+{
+	if (GetOwner()->Implements<UWorldPartitionActorLoaderInterface>())
+	{
+		IWorldPartitionActorLoaderInterface::ILoaderAdapter* LoaderAdapter = Cast<IWorldPartitionActorLoaderInterface>(GetOwner())->GetLoaderAdapter();
+		LoaderAdapter->Unload();
+	}
+}
+
+bool UProceduralFoliageComponent::IsSimulatedRegionLoaded()
+{
+	if (GetOwner()->Implements<UWorldPartitionActorLoaderInterface>())
+	{
+		IWorldPartitionActorLoaderInterface::ILoaderAdapter* LoaderAdapter = Cast<IWorldPartitionActorLoaderInterface>(GetOwner())->GetLoaderAdapter();
+		return LoaderAdapter->IsLoaded();
+	}
+	return true;
+}
+#endif
+
 FVector UProceduralFoliageComponent::GetWorldPosition() const
 {
 	FBox Bounds = GetBounds();
@@ -119,14 +149,7 @@ FVector UProceduralFoliageComponent::GetWorldPosition() const
 bool UProceduralFoliageComponent::ExecuteSimulation(TArray<FDesiredFoliageInstance>& OutInstances)
 {
 #if WITH_EDITOR
-
-	// In World Partition, load Editor Cells intersecting bounds affected by ProceduralFoliageCompoment
-	if (UWorldPartition* WorldPartition = GetWorld()->GetWorldPartition())
-	{
-		FBox Bounds = GetBounds();
-		check(Bounds.IsValid);
-		WorldPartition->LoadEditorCells(Bounds, false);
-	}
+	LoadSimulatedRegion();
 
 	FBodyInstance* BoundsBodyInstance = GetBoundsBodyInstance();
 	if (FoliageSpawner)
diff --git a/Engine/Source/Runtime/Foliage/Private/ProceduralFoliageVolume.cpp b/Engine/Source/Runtime/Foliage/Private/ProceduralFoliageVolume.cpp
--- a/Engine/Source/Runtime/Foliage/Private/ProceduralFoliageVolume.cpp
+++ b/Engine/Source/Runtime/Foliage/Private/ProceduralFoliageVolume.cpp
@@ -4,6 +4,7 @@
 #include "Components/BrushComponent.h"
 #include "ProceduralFoliageComponent.h"
 #include "ProceduralFoliageSpawner.h"
+#include "WorldPartition/LoaderAdapter/LoaderAdapterActor.h"
 
 AProceduralFoliageVolume::AProceduralFoliageVolume(const FObjectInitializer& ObjectInitializer)
 : Super(ObjectInitializer)
@@ -20,9 +21,31 @@ AProceduralFoliageVolume::AProceduralFoliageVolume(const FObjectInitializer& Obj
 		// That means during streaming we'll get a huge hitch for UpdateOverlaps
 		MyBrushComponent->SetGenerateOverlapEvents(false);
 	}
+
+#if WITH_EDITOR
+	if (!IsTemplate() && GetWorld() && GetWorld()->GetWorldPartition())
+	{
+		WorldPartitionActorLoader = new FLoaderAdapterActor(this);
+	}
+#endif
 }
 
 #if WITH_EDITOR
+void AProceduralFoliageVolume::BeginDestroy()
+{
+	if (WorldPartitionActorLoader)
+	{
+		delete WorldPartitionActorLoader;
+		WorldPartitionActorLoader = nullptr;
+	}
+
+	Super::BeginDestroy();
+}
+
+IWorldPartitionActorLoaderInterface::ILoaderAdapter* AProceduralFoliageVolume::GetLoaderAdapter()
+{
+	return WorldPartitionActorLoader;
+}
 
 void AProceduralFoliageVolume::PostEditImport()
 {
@@ -40,5 +63,4 @@ bool AProceduralFoliageVolume::GetReferencedContentObjects(TArray<UObject*>& Obj
 	}
 	return true;
 }
-
 #endif
diff --git a/Engine/Source/Runtime/Foliage/Public/ProceduralFoliageComponent.h b/Engine/Source/Runtime/Foliage/Public/ProceduralFoliageComponent.h
--- a/Engine/Source/Runtime/Foliage/Public/ProceduralFoliageComponent.h
+++ b/Engine/Source/Runtime/Foliage/Public/ProceduralFoliageComponent.h
@@ -115,6 +115,12 @@ class FOLIAGE_API UProceduralFoliageComponent : public UActorComponent
 	void SetSpawningVolume(AVolume* InSpawningVolume) { SpawningVolume = InSpawningVolume; }
 	const FGuid& GetProceduralGuid() const { return ProceduralGuid; }
 
+#if WITH_EDITOR
+	void LoadSimulatedRegion();
+	void UnloadSimulatedRegion();
+	bool IsSimulatedRegionLoaded();
+#endif
+
 private:
 	/** Does all the actual work of executing the procedural foliage simulation */
 	bool ExecuteSimulation(TArray<FDesiredFoliageInstance>& OutFoliageInstances);
diff --git a/Engine/Source/Runtime/Foliage/Public/ProceduralFoliageVolume.h b/Engine/Source/Runtime/Foliage/Public/ProceduralFoliageVolume.h
--- a/Engine/Source/Runtime/Foliage/Public/ProceduralFoliageVolume.h
+++ b/Engine/Source/Runtime/Foliage/Public/ProceduralFoliageVolume.h
@@ -5,12 +5,14 @@
 #include "CoreMinimal.h"
 #include "UObject/ObjectMacros.h"
 #include "GameFramework/Volume.h"
+#include "WorldPartition/WorldPartitionActorLoaderInterface.h"
 #include "ProceduralFoliageVolume.generated.h"
 
 class UProceduralFoliageComponent;
+class FLoaderAdapterActor;
 
 UCLASS()
-class FOLIAGE_API AProceduralFoliageVolume: public AVolume
+class FOLIAGE_API AProceduralFoliageVolume: public AVolume, public IWorldPartitionActorLoaderInterface
 {
 	GENERATED_UCLASS_BODY()
 
@@ -18,10 +20,18 @@ class FOLIAGE_API AProceduralFoliageVolume: public AVolume
 	TObjectPtr<UProceduralFoliageComponent> ProceduralComponent;
 
 #if WITH_EDITOR
-
-	// UObject interface
+	//~ Begin UObject Interface
+	virtual void BeginDestroy() override;
 	virtual void PostEditImport() override;
+	//~ End UObject Interface
 
 	virtual bool GetReferencedContentObjects(TArray<UObject*>& Objects) const override;
+
+	//~ Begin IWorldPartitionActorLoaderInterface interface
+	virtual ILoaderAdapter* GetLoaderAdapter() override;
+	//~ End IWorldPartitionActorLoaderInterface interface
+
+private:
+	FLoaderAdapterActor* WorldPartitionActorLoader;
 #endif
 };
