World Partition: added option to convert loading regions to a location volume actor in the World Partition UI.

#rnx

diff --git a/Engine/Source/Editor/UnrealEd/Classes/ActorFactories/ActorFactory.h b/Engine/Source/Editor/UnrealEd/Classes/ActorFactories/ActorFactory.h
--- a/Engine/Source/Editor/UnrealEd/Classes/ActorFactories/ActorFactory.h
+++ b/Engine/Source/Editor/UnrealEd/Classes/ActorFactories/ActorFactory.h
@@ -20,6 +20,8 @@ struct FAssetData;
 class UBlueprint;
 class ULevel;
 class UInstancedPlacemenClientSettings;
+class AVolume;
+class UBrushBuilder;
 
 UCLASS(collapsecategories, hidecategories=Object, editinlinenew, config=Editor, abstract, transient)
 class UNREALED_API UActorFactory : public UObject, public IAssetFactoryInterface
@@ -97,6 +99,8 @@ class UNREALED_API UActorFactory : public UObject, public IAssetFactoryInterface
 	virtual UInstancedPlacemenClientSettings* FactorySettingsObjectForPlacement(const FAssetData& InAssetData, const FPlacementOptions& InPlacementOptions) override;
 	// End IAssetFactoryInterface Interface
 
+	static void CreateBrushForVolumeActor(AVolume* NewActor, UBrushBuilder* BrushBuilder);
+
 protected:
 
 	/** Validates the input params for SpawnActor and returns the appropriate level to use depending on whether InLevel and/or InSpawnParams.OverrideLevel is passed : */
diff --git a/Engine/Source/Editor/UnrealEd/Private/Factories/ActorFactory.cpp b/Engine/Source/Editor/UnrealEd/Private/Factories/ActorFactory.cpp
--- a/Engine/Source/Editor/UnrealEd/Private/Factories/ActorFactory.cpp
+++ b/Engine/Source/Editor/UnrealEd/Private/Factories/ActorFactory.cpp
@@ -1915,8 +1915,8 @@ void UActorFactoryVectorFieldVolume::PostSpawnActor( UObject* Asset, AActor* New
 /*-----------------------------------------------------------------------------
 CreateBrushForVolumeActor
 -----------------------------------------------------------------------------*/
-// Helper function for the volume actor factories, not sure where it should live
-void CreateBrushForVolumeActor( AVolume* NewActor, UBrushBuilder* BrushBuilder )
+// Helper function for the volume actor factories
+void UActorFactory::CreateBrushForVolumeActor( AVolume* NewActor, UBrushBuilder* BrushBuilder )
 {
 	if ( NewActor != nullptr )
 	{
diff --git a/Engine/Source/Editor/WorldPartitionEditor/Private/WorldPartition/SWorldPartitionEditorGrid2D.cpp b/Engine/Source/Editor/WorldPartitionEditor/Private/WorldPartition/SWorldPartitionEditorGrid2D.cpp
--- a/Engine/Source/Editor/WorldPartitionEditor/Private/WorldPartition/SWorldPartitionEditorGrid2D.cpp
+++ b/Engine/Source/Editor/WorldPartitionEditor/Private/WorldPartition/SWorldPartitionEditorGrid2D.cpp
@@ -12,6 +12,8 @@
 #include "WorldPartition/WorldPartitionActorLoaderInterface.h"
 #include "WorldPartition/WorldPartitionEditorPerProjectUserSettings.h"
 #include "WorldPartition/LoaderAdapter/LoaderAdapterShape.h"
+#include "ActorFactories/ActorFactory.h"
+#include "Builders/CubeBuilder.h"
 #include "Editor/GroupActor.h"
 #include "Editor/EditorEngine.h"
 #include "Fonts/FontMeasure.h"
@@ -20,6 +22,7 @@
 #include "Rendering/SlateRenderer.h"
 #include "WorldBrowserModule.h"
 #include "LevelEditorViewport.h"
+#include "LocationVolume.h"
 #include "Algo/Transform.h"
 
 #define LOCTEXT_NAMESPACE "WorldPartitionEditor"
@@ -30,7 +33,7 @@ static bool IsBoundsSelected(const FBox& SelectBox, const FBox& Bounds)
 }
 
 template <class T>
-void ForEachSelectedLoaderAdapters(const UWorldPartition* WorldPartition, const FBox& SelectBox, T Func)
+void ForEachSelectedEditorLoaderAdapters(const UWorldPartition* WorldPartition, const FBox& SelectBox, T Func)
 {
 	for (IWorldPartitionActorLoaderInterface::ILoaderAdapter* LoaderAdapter : WorldPartition->GetRegisteredEditorLoaderAdapters())
 	{
@@ -45,7 +48,11 @@ void ForEachSelectedLoaderAdapters(const UWorldPartition* WorldPartition, const
 			}
 		}
 	}
+};
 
+template <class T>
+void ForEachSelectedWorldLoaderAdapters(const UWorldPartition* WorldPartition, const FBox& SelectBox, T Func)
+{
 	WorldPartition->EditorHash->ForEachIntersectingActor(SelectBox, [&](FWorldPartitionActorDesc* ActorDesc)
 	{
 		if (AActor* Actor = ActorDesc->GetActor())
@@ -67,6 +74,13 @@ void ForEachSelectedLoaderAdapters(const UWorldPartition* WorldPartition, const
 	});
 }
 
+template <class T>
+void ForEachSelectedLoaderAdapters(const UWorldPartition* WorldPartition, const FBox& SelectBox, T Func)
+{
+	ForEachSelectedEditorLoaderAdapters<T>(WorldPartition, SelectBox, Func);
+	ForEachSelectedWorldLoaderAdapters<T>(WorldPartition, SelectBox, Func);
+}
+
 class FWorldPartitionActorDescViewBoundsProxy : public FWorldPartitionActorDescView
 {
 public:
@@ -103,8 +117,9 @@ SWorldPartitionEditorGrid2D::FEditorCommands::FEditorCommands()
 void SWorldPartitionEditorGrid2D::FEditorCommands::RegisterCommands()
 {
 	UI_COMMAND(CreateRegionFromSelection, "Create Loading Region From Selection", "Create a loading region from the selection.", EUserInterfaceActionType::Button, FInputChord());
-	UI_COMMAND(LoadSelectedRegions, "Load Selected Regions", "Load the regions touched by the selection.", EUserInterfaceActionType::Button, FInputChord());
-	UI_COMMAND(UnloadSelectedRegions, "Unload Selected Regions", "Unload the regions touched by the selection.", EUserInterfaceActionType::Button, FInputChord());
+	UI_COMMAND(LoadSelectedRegions, "Load Selected Regions", "Load the selected regions.", EUserInterfaceActionType::Button, FInputChord());
+	UI_COMMAND(UnloadSelectedRegions, "Unload Selected Regions", "Unload the selected regions.", EUserInterfaceActionType::Button, FInputChord());
+	UI_COMMAND(ConvertSelectedRegionsToActors, "Convert Selected Regions To Actors", "Convert the selected regions to actors.", EUserInterfaceActionType::Button, FInputChord());
 	UI_COMMAND(MoveCameraHere, "Move Camera Here", "Move the camera to the selected position.", EUserInterfaceActionType::Button, FInputChord());
 }
 
@@ -244,12 +259,24 @@ void SWorldPartitionEditorGrid2D::Construct(const FArguments& InArgs)
 		return bResult;
 	};
 
+	auto CanConvertSelectedRegionsToActors = [this]()
+	{
+		bool bResult = false;
+		ForEachSelectedEditorLoaderAdapters(WorldPartition, SelectBox, [&bResult](IWorldPartitionActorLoaderInterface::ILoaderAdapter* LoaderAdapter)
+		{
+			bResult = true;
+			return false;
+		});
+		return bResult;
+	};
+
 	auto CanLoadSelectedRegions = [this, CanLoadUnloadSelectedRegions]() { return CanLoadUnloadSelectedRegions(true); };
 	auto CanUnloadSelectedRegions = [this, CanLoadUnloadSelectedRegions]() { return CanLoadUnloadSelectedRegions(false); };
 
 	ActionList.MapAction(Commands.CreateRegionFromSelection, FExecuteAction::CreateSP(this, &SWorldPartitionEditorGrid2D::CreateRegionFromSelection), FCanExecuteAction::CreateLambda(CanCreateRegionFromSelection));
 	ActionList.MapAction(Commands.LoadSelectedRegions, FExecuteAction::CreateSP(this, &SWorldPartitionEditorGrid2D::LoadSelectedRegions), FCanExecuteAction::CreateLambda(CanLoadSelectedRegions));
 	ActionList.MapAction(Commands.UnloadSelectedRegions, FExecuteAction::CreateSP(this, &SWorldPartitionEditorGrid2D::UnloadSelectedRegions), FCanExecuteAction::CreateLambda(CanUnloadSelectedRegions));
+	ActionList.MapAction(Commands.ConvertSelectedRegionsToActors, FExecuteAction::CreateSP(this, &SWorldPartitionEditorGrid2D::ConvertSelectedRegionsToActors), FCanExecuteAction::CreateLambda(CanConvertSelectedRegionsToActors));
 	ActionList.MapAction(Commands.MoveCameraHere, FExecuteAction::CreateSP(this, &SWorldPartitionEditorGrid2D::MoveCameraHere));
 }
 
@@ -309,6 +336,37 @@ void SWorldPartitionEditorGrid2D::UnloadSelectedRegions()
 	Refresh();
 }
 
+void SWorldPartitionEditorGrid2D::ConvertSelectedRegionsToActors()
+{
+	TArray<IWorldPartitionActorLoaderInterface::ILoaderAdapter*> LoaderAdaptersToRelease;
+	ForEachSelectedEditorLoaderAdapters(WorldPartition, SelectBox, [this, &LoaderAdaptersToRelease](IWorldPartitionActorLoaderInterface::ILoaderAdapter* LoaderAdapter)
+	{
+		const FBox LoaderVolumeBox(*LoaderAdapter->GetBoundingBox());
+
+		ALocationVolume* LocationVolume = World->SpawnActor<ALocationVolume>(LoaderVolumeBox.GetCenter(), FRotator::ZeroRotator);
+
+		UCubeBuilder* Builder = NewObject<UCubeBuilder>();
+		Builder->X = 1.0f;
+		Builder->Y = 1.0f;
+		Builder->Z = 1.0f;
+		UActorFactory::CreateBrushForVolumeActor(LocationVolume, Builder);
+
+		LocationVolume->GetRootComponent()->SetWorldScale3D(LoaderVolumeBox.GetSize());
+
+		LocationVolume->GetLoaderAdapter()->Load();
+		
+		LoaderAdapter->Unload();
+		LoaderAdaptersToRelease.Add(LoaderAdapter);
+
+		return true;
+	});
+
+	for (IWorldPartitionActorLoaderInterface::ILoaderAdapter* LoaderAdapterToRelease : LoaderAdaptersToRelease)
+	{
+		WorldPartition->ReleaseEditorLoaderAdapter(LoaderAdapterToRelease);
+	}
+}
+
 void SWorldPartitionEditorGrid2D::MoveCameraHere()
 {
 	for (FLevelEditorViewportClient* LevelVC : GEditor->GetLevelViewportClients())
@@ -372,6 +430,8 @@ FReply SWorldPartitionEditorGrid2D::OnMouseButtonUp(const FGeometry& MyGeometry,
 				MenuBuilder.AddMenuSeparator();
 				MenuBuilder.AddMenuEntry(Commands.LoadSelectedRegions);
 				MenuBuilder.AddMenuEntry(Commands.UnloadSelectedRegions);
+				MenuBuilder.AddMenuSeparator();
+				MenuBuilder.AddMenuEntry(Commands.ConvertSelectedRegionsToActors);
 			MenuBuilder.EndSection();
 
 			MenuBuilder.BeginSection(NAME_None, LOCTEXT("WorldPartitionMisc", "Misc"));
@@ -550,6 +610,33 @@ uint32 SWorldPartitionEditorGrid2D::PaintActors(const FGeometry& AllottedGeometr
 		}
 	}
 
+	auto DrawActorLabel = [&OutDrawElements, &LayerId, &AllottedGeometry](const FString& ActorLabel, const FBox2D& ActorViewBox, const FPaintGeometry& ActorGeometry, const FLinearColor& Color, const FSlateFontInfo& Font)
+	{
+		const FVector2D LabelTextSize = FSlateApplication::Get().GetRenderer()->GetFontMeasureService()->Measure(ActorLabel, Font);
+
+		if (LabelTextSize.X > 0)
+		{
+			const FVector2D ActorViewBoxCenter = ActorViewBox.GetCenter();
+			const FVector2D LabelTextPos = ActorViewBoxCenter - LabelTextSize * 0.5f;
+			const float LabelColorGradient = FMath::Clamp(ActorGeometry.GetLocalSize().X / LabelTextSize.X - 1.0f, 0.0f, 1.0f);
+
+			if (LabelColorGradient > 0.0f)
+			{
+				const FLinearColor LabelColor(Color.R, Color.G, Color.B, Color.A * LabelColorGradient);
+
+				FSlateDrawElement::MakeText(
+					OutDrawElements,
+					++LayerId,
+					AllottedGeometry.ToPaintGeometry(LabelTextPos, FVector2D(1,1)),
+					ActorLabel,
+					Font,
+					ESlateDrawEffect::None,
+					LabelColor
+				);
+			}
+		}
+	};
+
 	if (AllLoaderAdapters.Num())
 	{
 		const FLinearColor LoadedActorColor(0.75f, 0.75f, 0.75f, 1.0f);
@@ -634,29 +721,8 @@ uint32 SWorldPartitionEditorGrid2D::PaintActors(const FGeometry& AllottedGeometr
 						// Label
 						{
 							const FString ActorLabel = *LoaderAdapter->GetLabel();
-							const FVector2D ActorViewBoxCenter = ActorViewBox.GetCenter();
-							const FVector2D LabelTextSize = FSlateApplication::Get().GetRenderer()->GetFontMeasureService()->Measure(ActorLabel, SmallLayoutFont);
-
-							if (LabelTextSize.X > 0)
-							{
-								const FVector2D LabelTextPos = ActorViewBoxCenter - LabelTextSize * 0.5f;					
-								const float LabelColorGradient = FMath::Clamp(ActorGeometry.GetLocalSize().X / LabelTextSize.X - 1.0f, 0.0f, 1.0f);					
-
-								if (LabelColorGradient > 0.0f)
-								{
-									const FLinearColor LabelColor(1.0f, 1.0f, 1.0f, LabelColorGradient * LoaderColorGradient * FullScreenColorGradient);
-
-									FSlateDrawElement::MakeText(
-										OutDrawElements,
-										++LayerId,
-										AllottedGeometry.ToPaintGeometry(LabelTextPos, FVector2D(1,1)),
-										ActorLabel,
-										SmallLayoutFont,
-										ESlateDrawEffect::None,
-										LabelColor
-									);
-								}
-							}
+							const FLinearColor LabelColor(1.0f, 1.0f, 1.0f, LoaderColorGradient * FullScreenColorGradient);
+							DrawActorLabel(ActorLabel, ActorViewBox, ActorGeometry, LabelColor, SmallLayoutFont);
 						}
 					}
 				}
@@ -708,18 +774,10 @@ uint32 SWorldPartitionEditorGrid2D::PaintActors(const FGeometry& AllottedGeometr
 				{
 					ActorColor = FLinearColor::Yellow;
 
-					FName ActorLabel = ActorDescView.GetActorLabel();
+					const FName ActorLabel = ActorDescView.GetActorLabel();
 					if (!ActorLabel.IsNone())
 					{
-						FSlateDrawElement::MakeText(
-							OutDrawElements,
-							++LayerId,
-							AllottedGeometry.ToPaintGeometry(TopLeft, FVector2D(1,1)),
-							ActorLabel.ToString(),
-							SmallLayoutFont,
-							ESlateDrawEffect::None,
-							ActorColor
-						);
+						DrawActorLabel(ActorLabel.ToString(), ActorViewBox, ActorGeometry, ActorColor, SmallLayoutFont);
 					}
 				}
 				else if ((SelectBox.GetVolume() > 0) && SelectBox.Intersect(ActorDescView.GetBounds()))
diff --git a/Engine/Source/Editor/WorldPartitionEditor/Public/WorldPartition/SWorldPartitionEditorGrid2D.h b/Engine/Source/Editor/WorldPartitionEditor/Public/WorldPartition/SWorldPartitionEditorGrid2D.h
--- a/Engine/Source/Editor/WorldPartitionEditor/Public/WorldPartition/SWorldPartitionEditorGrid2D.h
+++ b/Engine/Source/Editor/WorldPartitionEditor/Public/WorldPartition/SWorldPartitionEditorGrid2D.h
@@ -18,6 +18,7 @@ protected:
 		TSharedPtr<FUICommandInfo> CreateRegionFromSelection;
 		TSharedPtr<FUICommandInfo> LoadSelectedRegions;
 		TSharedPtr<FUICommandInfo> UnloadSelectedRegions;
+		TSharedPtr<FUICommandInfo> ConvertSelectedRegionsToActors;
 		TSharedPtr<FUICommandInfo> MoveCameraHere;
 
 		/**
@@ -35,6 +36,7 @@ public:
 	virtual void CreateRegionFromSelection();
 	virtual void LoadSelectedRegions();
 	virtual void UnloadSelectedRegions();
+	virtual void ConvertSelectedRegionsToActors();
 	void MoveCameraHere();
 
 	virtual FReply OnMouseButtonDown(const FGeometry& MyGeometry, const FPointerEvent& MouseEvent) override;
