World Partition:
- Added assert that the provided world is not a game world in IWorldPartitionActorLoaderInterface::ILoaderAdapter constructor.
- Apply the ALocationVolume.DebugColor to the World Partition display for loader adapters UI.

#rnx

diff --git a/Engine/Source/Editor/WorldPartitionEditor/Private/WorldPartition/SWorldPartitionEditorGrid2D.cpp b/Engine/Source/Editor/WorldPartitionEditor/Private/WorldPartition/SWorldPartitionEditorGrid2D.cpp
--- a/Engine/Source/Editor/WorldPartitionEditor/Private/WorldPartition/SWorldPartitionEditorGrid2D.cpp
+++ b/Engine/Source/Editor/WorldPartitionEditor/Private/WorldPartition/SWorldPartitionEditorGrid2D.cpp
@@ -572,8 +572,9 @@ uint32 SWorldPartitionEditorGrid2D::PaintActors(const FGeometry& AllottedGeometr
 	TSet<FWorldPartitionActorDescViewBoundsProxy> ActorDescList;
 	
 	TArray<IWorldPartitionActorLoaderInterface::ILoaderAdapter*> AllLoaderAdapters;
-	for (IWorldPartitionActorLoaderInterface::ILoaderAdapter* LoaderAdapter: WorldPartition->GetRegisteredEditorLoaderAdapters())
+	for (IWorldPartitionActorLoaderInterface::ILoaderAdapter* LoaderAdapter : WorldPartition->GetRegisteredEditorLoaderAdapters())
 	{
+		check(LoaderAdapter);
 		AllLoaderAdapters.Add(LoaderAdapter);
 	}
 
@@ -592,8 +593,11 @@ uint32 SWorldPartitionEditorGrid2D::PaintActors(const FGeometry& AllottedGeometr
 		if (ActorDesc->GetActorNativeClass()->ImplementsInterface(UWorldPartitionActorLoaderInterface::StaticClass()))
 		{
 			if (AActor* Actor = ActorDescViewProxy.GetActor())
-			{				
-				AllLoaderAdapters.Add(Cast<IWorldPartitionActorLoaderInterface>(Actor)->GetLoaderAdapter());
+			{
+				if (IWorldPartitionActorLoaderInterface::ILoaderAdapter* LoaderAdapter = Cast<IWorldPartitionActorLoaderInterface>(Actor)->GetLoaderAdapter())
+				{
+					AllLoaderAdapters.Add(LoaderAdapter);
+				}
 			}
 		}
 	});
@@ -696,8 +700,9 @@ uint32 SWorldPartitionEditorGrid2D::PaintActors(const FGeometry& AllottedGeometr
 
 						// Outline
 						{
+							const FLinearColor LoaderColor = LoaderAdapter->GetColor().IsSet() ? *LoaderAdapter->GetColor() : FColor::White;
+							const FLinearColor OutlineColor(LoaderColor.R, LoaderColor.G, LoaderColor.B, LoaderColorGradient * FullScreenColorGradient);
 							const bool IsHighlighted = IsBoundsSelected(SelectBox, ActorBounds);
-							const FLinearColor OutlineColor(IsHighlighted ? 1.0f : 0.0f, IsHighlighted ? 1.0f : 0.0f, 0.0f, LoaderColorGradient * FullScreenColorGradient);
 
 							LinePoints[0] = TopLeft;
 							LinePoints[1] = TopRight;
@@ -714,7 +719,7 @@ uint32 SWorldPartitionEditorGrid2D::PaintActors(const FGeometry& AllottedGeometr
 								ESlateDrawEffect::None,
 								OutlineColor,
 								true,
-								2.0f
+								IsHighlighted ? 4.0f : 2.0f
 							);
 						}
 
@@ -1088,4 +1093,4 @@ void SWorldPartitionEditorGrid2D::UpdateSelection()
 	}
 }
 
-#undef LOCTEXT_NAMESPACE
+#undef LOCTEXT_NAMESPACE
\ No newline at end of file
diff --git a/Engine/Source/Runtime/Engine/Private/LocationVolume.cpp b/Engine/Source/Runtime/Engine/Private/LocationVolume.cpp
--- a/Engine/Source/Runtime/Engine/Private/LocationVolume.cpp
+++ b/Engine/Source/Runtime/Engine/Private/LocationVolume.cpp
@@ -5,6 +5,23 @@
 #include "Components/BrushComponent.h"
 #include "WorldPartition/LoaderAdapter/LoaderAdapterActor.h"
 
+#if WITH_EDITOR
+class ENGINE_API FLoaderAdapterLocationVolumeActor : public FLoaderAdapterActor
+{
+public:
+	FLoaderAdapterLocationVolumeActor(AActor* InActor)
+		: FLoaderAdapterActor(InActor)
+	{}
+
+	//~ Begin IWorldPartitionActorLoaderInterface::ILoader interface
+	virtual TOptional<FColor> GetColor() const override
+	{
+		return CastChecked<ALocationVolume>(Actor)->DebugColor;
+	}
+	//~ End IWorldPartitionActorLoaderInterface::ILoader interface
+};
+#endif
+
 ALocationVolume::ALocationVolume(const FObjectInitializer& ObjectInitializer)
 : Super(ObjectInitializer)
 {
@@ -12,43 +29,39 @@ ALocationVolume::ALocationVolume(const FObjectInitializer& ObjectInitializer)
 
 #if WITH_EDITORONLY_DATA
 	bIsSpatiallyLoaded = false;
-	bIsAutoLoad = false;
-#endif
-
-#if WITH_EDITOR
-	if (!IsTemplate() && GetWorld() && GetWorld()->GetWorldPartition())
-	{
-		WorldPartitionActorLoader = new FLoaderAdapterActor(this);
-	}
+	DebugColor = FColor::White;
+	bIsAutoLoad = false;	
 #endif
 }
 
+#if WITH_EDITOR
 void ALocationVolume::PostRegisterAllComponents()
 {
 	Super::PostRegisterAllComponents();
 
-#if WITH_EDITOR
-	if (WorldPartitionActorLoader && bIsAutoLoad)
+	if (!GetWorld()->IsGameWorld() && GetWorld()->IsPartitionedWorld() && !WorldPartitionActorLoader)
 	{
-		WorldPartitionActorLoader->Load();
+		WorldPartitionActorLoader = new FLoaderAdapterLocationVolumeActor(this);
+
+		if (bIsAutoLoad)
+		{
+			bIsAutoLoad = false;
+			WorldPartitionActorLoader->Load();
+		}
 	}
-#endif
 }
 
 void ALocationVolume::BeginDestroy()
 {
-#if WITH_EDITOR
 	if (WorldPartitionActorLoader)
 	{
 		delete WorldPartitionActorLoader;
 		WorldPartitionActorLoader = nullptr;
 	}
-#endif
 
 	Super::BeginDestroy();
 }
 
-#if WITH_EDITOR
 IWorldPartitionActorLoaderInterface::ILoaderAdapter* ALocationVolume::GetLoaderAdapter()
 {
 	return WorldPartitionActorLoader;
diff --git a/Engine/Source/Runtime/Engine/Private/WorldPartition/WorldPartitionActorLoaderInterface.cpp b/Engine/Source/Runtime/Engine/Private/WorldPartition/WorldPartitionActorLoaderInterface.cpp
--- a/Engine/Source/Runtime/Engine/Private/WorldPartition/WorldPartitionActorLoaderInterface.cpp
+++ b/Engine/Source/Runtime/Engine/Private/WorldPartition/WorldPartitionActorLoaderInterface.cpp
@@ -26,7 +26,9 @@ IWorldPartitionActorLoaderInterface::ILoaderAdapter::ILoaderAdapter(UWorld* InWo
 	: World(InWorld)
 	, bLoaded(false)
 	, bUserCreated(false)
-{}
+{
+	check(!World->IsGameWorld());
+}
 
 IWorldPartitionActorLoaderInterface::ILoaderAdapter::~ILoaderAdapter()
 {
diff --git a/Engine/Source/Runtime/Engine/Public/LocationVolume.h b/Engine/Source/Runtime/Engine/Public/LocationVolume.h
--- a/Engine/Source/Runtime/Engine/Public/LocationVolume.h
+++ b/Engine/Source/Runtime/Engine/Public/LocationVolume.h
@@ -20,19 +20,19 @@ class ENGINE_API ALocationVolume : public AVolume, public IWorldPartitionActorLo
 
 public:
 	//~ Begin UObject Interface
+#if WITH_EDITOR
 	virtual void BeginDestroy() override;
+#endif
 	//~ End UObject Interface
 
 	//~ Begin AActor Interface
 	virtual bool IsEditorOnly() const override { return !bIsRuntime; }
-	virtual void PostRegisterAllComponents();
 #if WITH_EDITOR
+	virtual void PostRegisterAllComponents();
 	virtual bool CanChangeIsSpatiallyLoadedFlag() const override { return false; }
 	virtual bool SupportsDataLayer() const override { return false; }
-#endif
 	//~ End AActor Interface
 
-#if WITH_EDITOR
 	//~ Begin IWorldPartitionActorLoaderInterface interface
 	virtual ILoaderAdapter* GetLoaderAdapter() override;
 	//~ End IWorldPartitionActorLoaderInterface interface
diff --git a/Engine/Source/Runtime/Engine/Public/WorldPartition/LoaderAdapter/LoaderAdapterActor.h b/Engine/Source/Runtime/Engine/Public/WorldPartition/LoaderAdapter/LoaderAdapterActor.h
--- a/Engine/Source/Runtime/Engine/Public/WorldPartition/LoaderAdapter/LoaderAdapterActor.h
+++ b/Engine/Source/Runtime/Engine/Public/WorldPartition/LoaderAdapter/LoaderAdapterActor.h
@@ -20,7 +20,7 @@ public:
 	virtual bool Intersect(const FBox& Box) const override;
 	//~ End IWorldPartitionActorLoaderInterface::ILoaderAdapterSpatial interface
 
-private:
+protected:
 	AActor* Actor;
 };
 #endif
diff --git a/Engine/Source/Runtime/Engine/Public/WorldPartition/WorldPartitionActorLoaderInterface.h b/Engine/Source/Runtime/Engine/Public/WorldPartition/WorldPartitionActorLoaderInterface.h
--- a/Engine/Source/Runtime/Engine/Public/WorldPartition/WorldPartitionActorLoaderInterface.h
+++ b/Engine/Source/Runtime/Engine/Public/WorldPartition/WorldPartitionActorLoaderInterface.h
@@ -37,6 +37,7 @@ public:
 		// Interface
 		virtual TOptional<FBox> GetBoundingBox() const { return TOptional<FBox>(); }
 		virtual TOptional<FString> GetLabel() const { return TOptional<FString>(); }
+		virtual TOptional<FColor> GetColor() const { return TOptional<FColor>(); }
 
 	protected:
 		virtual bool Intersect(const FBox& Box) const =0;
diff --git a/Engine/Source/Runtime/Foliage/Private/ProceduralFoliageVolume.cpp b/Engine/Source/Runtime/Foliage/Private/ProceduralFoliageVolume.cpp
--- a/Engine/Source/Runtime/Foliage/Private/ProceduralFoliageVolume.cpp
+++ b/Engine/Source/Runtime/Foliage/Private/ProceduralFoliageVolume.cpp
@@ -21,16 +21,19 @@ AProceduralFoliageVolume::AProceduralFoliageVolume(const FObjectInitializer& Obj
 		// That means during streaming we'll get a huge hitch for UpdateOverlaps
 		MyBrushComponent->SetGenerateOverlapEvents(false);
 	}
+}
 
 #if WITH_EDITOR
-	if (!IsTemplate() && GetWorld() && GetWorld()->GetWorldPartition())
+void AProceduralFoliageVolume::PostRegisterAllComponents()
+{
+	Super::PostRegisterAllComponents();
+
+	if (!GetWorld()->IsGameWorld() && GetWorld()->IsPartitionedWorld() && !WorldPartitionActorLoader)
 	{
 		WorldPartitionActorLoader = new FLoaderAdapterActor(this);
 	}
-#endif
 }
 
-#if WITH_EDITOR
 void AProceduralFoliageVolume::BeginDestroy()
 {
 	if (WorldPartitionActorLoader)
diff --git a/Engine/Source/Runtime/Foliage/Public/ProceduralFoliageVolume.h b/Engine/Source/Runtime/Foliage/Public/ProceduralFoliageVolume.h
--- a/Engine/Source/Runtime/Foliage/Public/ProceduralFoliageVolume.h
+++ b/Engine/Source/Runtime/Foliage/Public/ProceduralFoliageVolume.h
@@ -20,6 +20,10 @@ class FOLIAGE_API AProceduralFoliageVolume: public AVolume, public IWorldPartiti
 	TObjectPtr<UProceduralFoliageComponent> ProceduralComponent;
 
 #if WITH_EDITOR
+	//~ Begin AActor Interface
+	virtual void PostRegisterAllComponents();
+	//~ End AActor Interface
+
 	//~ Begin UObject Interface
 	virtual void BeginDestroy() override;
 	virtual void PostEditImport() override;
