Update World Partition link with right/new/changed website path. Another change on WorldPartition.cpp will come by RM when 5.0.3 gates open.

#rnx

diff --git a/Engine/Source/Editor/WorldPartitionEditor/Private/WorldPartition/Customizations/WorldPartitionDetailsCustomization.cpp b/Engine/Source/Editor/WorldPartitionEditor/Private/WorldPartition/Customizations/WorldPartitionDetailsCustomization.cpp
--- a/Engine/Source/Editor/WorldPartitionEditor/Private/WorldPartition/Customizations/WorldPartitionDetailsCustomization.cpp
+++ b/Engine/Source/Editor/WorldPartitionEditor/Private/WorldPartition/Customizations/WorldPartitionDetailsCustomization.cpp
@@ -98,7 +98,7 @@ void FWorldPartitionDetails::HandleWorldPartitionEnableStreamingChanged(ECheckBo
 				return;
 			}
 			
-			FMessageDialog::Open(EAppMsgType::Ok, LOCTEXT("WorldPartitionEnableStreamingDialolg", "Please refer to https://docs.unrealengine.com/5.0/en-US/building-virtual-worlds/world-partition for how to set up streaming."));
+			FMessageDialog::Open(EAppMsgType::Ok, LOCTEXT("WorldPartitionEnableStreamingDialolg", "Please refer to https://docs.unrealengine.com/5.0/en-US/world-partition-in-unreal-engine/ for how to set up streaming."));
 			WorldPartition->bStreamingWasEnabled = true;
 		}
 
