World Partition: remove code related to editing the editor hash cell size.

#rnx

diff --git a/Engine/Source/Editor/WorldPartitionEditor/Private/WorldPartition/Customizations/WorldPartitionDetailsCustomization.cpp b/Engine/Source/Editor/WorldPartitionEditor/Private/WorldPartition/Customizations/WorldPartitionDetailsCustomization.cpp
--- a/Engine/Source/Editor/WorldPartitionEditor/Private/WorldPartition/Customizations/WorldPartitionDetailsCustomization.cpp
+++ b/Engine/Source/Editor/WorldPartitionEditor/Private/WorldPartition/Customizations/WorldPartitionDetailsCustomization.cpp
@@ -61,30 +61,6 @@ void FWorldPartitionDetails::CustomizeDetails(IDetailLayoutBuilder& DetailBuilde
 				.Visibility(TAttribute<EVisibility>::CreateLambda([this]() { return WorldPartition.IsValid() && WorldPartition->IsStreamingEnabled() ? EVisibility::Visible : EVisibility::Hidden; }));
 		}
 	}
-
-	WorldPartitionCategory.AddCustomRow(LOCTEXT("EditorCellSizeRow", "Editor Cell Size"), true)
-		.NameContent()
-		[
-			SNew(STextBlock)
-			.Text(LOCTEXT("WorldPartitionCellSize", "Editor Cell Size"))
-			.ToolTipText(LOCTEXT("WorldPartitionEditorCellSize_ToolTip", "Set the world partition editor cell size, will take effect on the next world reload."))
-			.Font(IDetailLayoutBuilder::GetDetailFont())
-		]
-		.ValueContent()
-		[
-			SNew(SHorizontalBox)
-			+SHorizontalBox::Slot()
-			.AutoWidth()
-			.VAlign(VAlign_Center)
-			[
-				SNew(SNumericEntryBox<uint32>)
-				.AllowSpin(false)
-				.MinSliderValue(100)
-				.MaxSliderValue(100000)
-				.OnValueChanged(this, &FWorldPartitionDetails::HandleWorldPartitionEditorCellSizeChanged)
-				.Value(this, &FWorldPartitionDetails::HandleWorldPartitionEditorCellSizeValue)
-			]
-		];		
 }
 
 void FWorldPartitionDetails::HandleWorldPartitionEnableStreamingChanged(ECheckBoxState CheckState)
@@ -113,14 +89,4 @@ void FWorldPartitionDetails::HandleWorldPartitionEnableStreamingChanged(ECheckBo
 	}
 }
 
-void FWorldPartitionDetails::HandleWorldPartitionEditorCellSizeChanged(uint32 NewValue)
-{
-	WorldPartition->SetEditorWantedCellSize(NewValue);
-}
-
-TOptional<uint32> FWorldPartitionDetails::HandleWorldPartitionEditorCellSizeValue() const
-{
-	return WorldPartition->GetWantedEditorCellSize();
-}
-
 #undef LOCTEXT_NAMESPACE
diff --git a/Engine/Source/Editor/WorldPartitionEditor/Private/WorldPartition/Customizations/WorldPartitionDetailsCustomization.h b/Engine/Source/Editor/WorldPartitionEditor/Private/WorldPartition/Customizations/WorldPartitionDetailsCustomization.h
--- a/Engine/Source/Editor/WorldPartitionEditor/Private/WorldPartition/Customizations/WorldPartitionDetailsCustomization.h
+++ b/Engine/Source/Editor/WorldPartitionEditor/Private/WorldPartition/Customizations/WorldPartitionDetailsCustomization.h
@@ -28,11 +28,5 @@ private:
 	// Callback for changes in the world partition enable streaming.
 	void HandleWorldPartitionEnableStreamingChanged(ECheckBoxState CheckState);
 
-	// Callback for changes in the world partition editor cell size.
-	void HandleWorldPartitionEditorCellSizeChanged(uint32 NewValue);
-
-	// Callback for getting the world partition editor cell size.
-	TOptional<uint32> HandleWorldPartitionEditorCellSizeValue() const;
-
 	TWeakObjectPtr<UWorldPartition> WorldPartition;
 };
\ No newline at end of file
diff --git a/Engine/Source/Runtime/Engine/Private/WorldPartition/WorldPartition.cpp b/Engine/Source/Runtime/Engine/Private/WorldPartition/WorldPartition.cpp
--- a/Engine/Source/Runtime/Engine/Private/WorldPartition/WorldPartition.cpp
+++ b/Engine/Source/Runtime/Engine/Private/WorldPartition/WorldPartition.cpp
@@ -1331,16 +1331,6 @@ void UWorldPartition::AppendAssetRegistryTags(TArray<FAssetRegistryTag>& OutTags
 	}
 }
 
-uint32 UWorldPartition::GetWantedEditorCellSize() const
-{
-	return EditorHash->GetWantedEditorCellSize();
-}
-
-void UWorldPartition::SetEditorWantedCellSize(uint32 InCellSize)
-{
-	EditorHash->SetEditorWantedCellSize(InCellSize);
-}
-
 AActor* UWorldPartition::PinActor(const FGuid& ActorGuid)
 {
 	if (PinnedActors)
diff --git a/Engine/Source/Runtime/Engine/Private/WorldPartition/WorldPartitionEditorSpatialHash.cpp b/Engine/Source/Runtime/Engine/Private/WorldPartition/WorldPartitionEditorSpatialHash.cpp
--- a/Engine/Source/Runtime/Engine/Private/WorldPartition/WorldPartitionEditorSpatialHash.cpp
+++ b/Engine/Source/Runtime/Engine/Private/WorldPartition/WorldPartitionEditorSpatialHash.cpp
@@ -348,28 +348,6 @@ int32 UWorldPartitionEditorSpatialHash::ForEachIntersectingCell(const FBox& Box,
 
 	return NumIntersecting;
 }
-
-uint32 UWorldPartitionEditorSpatialHash::GetWantedEditorCellSize() const
-{
-	return WantedCellSize ? WantedCellSize : CellSize;
-}
-
-void UWorldPartitionEditorSpatialHash::SetEditorWantedCellSize(uint32 InCellSize)
-{
-	Modify();
-	WantedCellSize = InCellSize;
-}
-
-void UWorldPartitionEditorSpatialHash::PostLoad()
-{
-	Super::PostLoad();
-
-	if (WantedCellSize && (CellSize != WantedCellSize))
-	{
-		CellSize = WantedCellSize;
-		WantedCellSize = 0;
-	}
-}
 #endif
 
 #undef LOCTEXT_NAMESPACE
diff --git a/Engine/Source/Runtime/Engine/Public/WorldPartition/WorldPartition.h b/Engine/Source/Runtime/Engine/Public/WorldPartition/WorldPartition.h
--- a/Engine/Source/Runtime/Engine/Public/WorldPartition/WorldPartition.h
+++ b/Engine/Source/Runtime/Engine/Public/WorldPartition/WorldPartition.h
@@ -164,9 +164,6 @@ public:
 
 	void AppendAssetRegistryTags(TArray<FAssetRegistryTag>& OutTags) const;
 
-	uint32 GetWantedEditorCellSize() const;
-	void SetEditorWantedCellSize(uint32 InCellSize);
-
 	// Actors pinning
 	AActor* PinActor(const FGuid& ActorGuid);
 	void UnpinActor(const FGuid& ActorGuid);
diff --git a/Engine/Source/Runtime/Engine/Public/WorldPartition/WorldPartitionEditorHash.h b/Engine/Source/Runtime/Engine/Public/WorldPartition/WorldPartitionEditorHash.h
--- a/Engine/Source/Runtime/Engine/Public/WorldPartition/WorldPartitionEditorHash.h
+++ b/Engine/Source/Runtime/Engine/Public/WorldPartition/WorldPartitionEditorHash.h
@@ -25,8 +25,5 @@ public:
 	virtual void UnhashActor(FWorldPartitionHandle& InActorHandle) PURE_VIRTUAL(UWorldPartitionEditorHash::UnhashActor, ;);
 
 	virtual int32 ForEachIntersectingActor(const FBox& Box, TFunctionRef<void(FWorldPartitionActorDesc*)> InOperation, bool bIncludeSpatiallyLoadedActors = true, bool bIncludeNonSpatiallyLoadedActors = true) PURE_VIRTUAL(UWorldPartitionEditorHash::ForEachIntersectingActor, return 0;);
-
-	virtual uint32 GetWantedEditorCellSize() const PURE_VIRTUAL(UWorldPartitionEditorHash::GetWantedEditorCellSize, return 0;);
-	virtual void SetEditorWantedCellSize(uint32 InCellSize) PURE_VIRTUAL(UWorldPartitionEditorHash::SetEditorWantedCellSize, );
 #endif
 };
diff --git a/Engine/Source/Runtime/Engine/Public/WorldPartition/WorldPartitionEditorSpatialHash.h b/Engine/Source/Runtime/Engine/Public/WorldPartition/WorldPartitionEditorSpatialHash.h
--- a/Engine/Source/Runtime/Engine/Public/WorldPartition/WorldPartitionEditorSpatialHash.h
+++ b/Engine/Source/Runtime/Engine/Public/WorldPartition/WorldPartitionEditorSpatialHash.h
@@ -186,10 +186,6 @@ class ENGINE_API UWorldPartitionEditorSpatialHash : public UWorldPartitionEditor
 public:
 	virtual ~UWorldPartitionEditorSpatialHash() {}
 
-	// UObject interface begin
-	virtual void PostLoad() override;
-	// UObject interface end
-
 	// UWorldPartitionEditorHash interface begin
 	virtual void Initialize() override;
 	virtual void SetDefaultValues() override;
@@ -201,9 +197,6 @@ public:
 	virtual void UnhashActor(FWorldPartitionHandle& InActorHandle) override;
 
 	virtual int32 ForEachIntersectingActor(const FBox& Box, TFunctionRef<void(FWorldPartitionActorDesc*)> InOperation, bool bIncludeSpatiallyLoadedActors = true, bool bIncludeNonSpatiallyLoadedActors = true) override;
-
-	virtual uint32 GetWantedEditorCellSize() const override;
-	virtual void SetEditorWantedCellSize(uint32 InCellSize) override;
 	// UWorldPartitionEditorHash interface end
 #endif
 
@@ -215,9 +208,6 @@ private:
 	UPROPERTY(Config)
 	int32 CellSize;
 
-	UPROPERTY()
-	int32 WantedCellSize;
-
 	// Dynamic sparse octree structure
 	typedef TTuple<FCellNode, TUniquePtr<FCell>> FCellNodeElement;
 	typedef TMap<FCellCoord, FCellNodeElement> FCellNodeHashLevel;
