Fix World Partition Convert Cmdlet when starting from a project containing a whitespace in their path.

#rnx

diff --git a/Engine/Source/Editor/WorldPartitionEditor/Private/WorldPartitionEditorModule.cpp b/Engine/Source/Editor/WorldPartitionEditor/Private/WorldPartitionEditorModule.cpp
--- a/Engine/Source/Editor/WorldPartitionEditor/Private/WorldPartitionEditorModule.cpp
+++ b/Engine/Source/Editor/WorldPartitionEditor/Private/WorldPartitionEditorModule.cpp
@@ -313,7 +313,7 @@ void FWorldPartitionEditorModule::RunCommandletAsExternalProcess(const FString&
 	FString ProjectPath = FPaths::IsProjectFilePathSet() ? FPaths::GetProjectFilePath() : FApp::GetProjectName();
 
 	FString Arguments;
-	Arguments += ProjectPath;
+	Arguments += TEXT("\"") + ProjectPath + TEXT('"');
 	Arguments += TEXT(" -BaseDir=\"") + FString(FPlatformProcess::BaseDir()) + TEXT('"');
 	Arguments += TEXT(" -Unattended");
 	Arguments += TEXT(" ") + InCommandletArgs;
