World Partition LWC support: removed WORLDPARTITION_MAX and replaced by corresponding WORLD_MAX/HALF_WORLD_MAX where it makes sense.

#rnx

diff --git a/Engine/Source/Editor/UnrealEd/Private/Commandlets/WorldPartitionConvertCommandlet.cpp b/Engine/Source/Editor/UnrealEd/Private/Commandlets/WorldPartitionConvertCommandlet.cpp
--- a/Engine/Source/Editor/UnrealEd/Private/Commandlets/WorldPartitionConvertCommandlet.cpp
+++ b/Engine/Source/Editor/UnrealEd/Private/Commandlets/WorldPartitionConvertCommandlet.cpp
@@ -163,7 +163,7 @@ UWorldPartitionConvertCommandlet::UWorldPartitionConvertCommandlet(const FObject
 	, ConversionSuffix(TEXT("_WP"))
 	, bConvertActorsNotReferencedByLevelScript(true)
 	, WorldOrigin(FVector::ZeroVector)
-	, WorldExtent(WORLDPARTITION_MAX * 0.5)
+	, WorldExtent(HALF_WORLD_MAX)
 	, LandscapeGridSize(4)
 	, DataLayerAssetFolder(TEXT("Game/DataLayers/"))
 	, DataLayerFactory(NewObject<UDataLayerFactory>())
diff --git a/Engine/Source/Editor/UnrealEd/Private/WorldPartition/WorldPartitionBuilder.cpp b/Engine/Source/Editor/UnrealEd/Private/WorldPartition/WorldPartitionBuilder.cpp
--- a/Engine/Source/Editor/UnrealEd/Private/WorldPartition/WorldPartitionBuilder.cpp
+++ b/Engine/Source/Editor/UnrealEd/Private/WorldPartition/WorldPartitionBuilder.cpp
@@ -292,7 +292,7 @@ bool UWorldPartitionBuilder::Run(UWorld* World, FPackageSourceControlHelper& Pac
 		
 		if (LoadingMode == ELoadingMode::EntireWorld)
 		{
-			CellInfo.Bounds = FBox(FVector(-WORLDPARTITION_MAX, -WORLDPARTITION_MAX, -WORLDPARTITION_MAX), FVector(WORLDPARTITION_MAX, WORLDPARTITION_MAX, WORLDPARTITION_MAX));
+			CellInfo.Bounds = FBox(FVector(-HALF_WORLD_MAX, -HALF_WORLD_MAX, -HALF_WORLD_MAX), FVector(HALF_WORLD_MAX, HALF_WORLD_MAX, HALF_WORLD_MAX));
 			LoaderAdapterShape = MakeUnique<FLoaderAdapterShape>(World, CellInfo.Bounds, TEXT("Loaded Region"));
 			LoaderAdapterShape->Load();
 		}
diff --git a/Engine/Source/Editor/UnrealEd/Private/WorldPartition/WorldPartitionNavigationDataBuilder.cpp b/Engine/Source/Editor/UnrealEd/Private/WorldPartition/WorldPartitionNavigationDataBuilder.cpp
--- a/Engine/Source/Editor/UnrealEd/Private/WorldPartition/WorldPartitionNavigationDataBuilder.cpp
+++ b/Engine/Source/Editor/UnrealEd/Private/WorldPartition/WorldPartitionNavigationDataBuilder.cpp
@@ -349,7 +349,7 @@ bool UWorldPartitionNavigationDataBuilder::GenerateNavigationData(UWorldPartitio
 				return false;
 			}
 
-			constexpr float HalfHeight = WORLDPARTITION_MAX * 0.5;
+			constexpr float HalfHeight = HALF_WORLD_MAX;
 			const FBox QueryBounds(FVector(CellBounds.Min.X, CellBounds.Min.Y, -HalfHeight), FVector(CellBounds.Max.X, CellBounds.Max.Y, HalfHeight));
 
 			if (!NavDataBounds.IsValid || !NavDataBounds.Intersect(QueryBounds))
diff --git a/Engine/Source/Editor/WorldPartitionEditor/Private/WorldPartition/SWorldPartitionEditorGrid2D.cpp b/Engine/Source/Editor/WorldPartitionEditor/Private/WorldPartition/SWorldPartitionEditorGrid2D.cpp
--- a/Engine/Source/Editor/WorldPartitionEditor/Private/WorldPartition/SWorldPartitionEditorGrid2D.cpp
+++ b/Engine/Source/Editor/WorldPartitionEditor/Private/WorldPartition/SWorldPartitionEditorGrid2D.cpp
@@ -282,7 +282,7 @@ void SWorldPartitionEditorGrid2D::Construct(const FArguments& InArgs)
 
 void SWorldPartitionEditorGrid2D::CreateRegionFromSelection()
 {
-	const FBox RegionBox(FVector(SelectBox.Min.X, SelectBox.Min.Y, -WORLDPARTITION_MAX), FVector(SelectBox.Max.X, SelectBox.Max.Y, WORLDPARTITION_MAX));
+	const FBox RegionBox(FVector(SelectBox.Min.X, SelectBox.Min.Y, -HALF_WORLD_MAX), FVector(SelectBox.Max.X, SelectBox.Max.Y, HALF_WORLD_MAX));
 	FLoaderAdapterShape* LoaderAdapter = WorldPartition->CreateEditorLoaderAdapter<FLoaderAdapterShape>(World, RegionBox, TEXT("Loaded Region"));
 	LoaderAdapter->SetUserCreated(true);
 	LoaderAdapter->Load();
@@ -569,7 +569,7 @@ uint32 SWorldPartitionEditorGrid2D::PaintActors(const FGeometry& AllottedGeometr
 {
 	const FBox2D ViewRect(FVector2D(ForceInitToZero), AllottedGeometry.GetLocalSize());
 	const FBox2D WorldViewRect(ScreenToWorld.TransformPoint(ViewRect.Min), ScreenToWorld.TransformPoint(ViewRect.Max));
-	const FBox ViewRectWorld(FVector(WorldViewRect.Min.X, WorldViewRect.Min.Y, -WORLDPARTITION_MAX), FVector(WorldViewRect.Max.X, WorldViewRect.Max.Y, WORLDPARTITION_MAX));
+	const FBox ViewRectWorld(FVector(WorldViewRect.Min.X, WorldViewRect.Min.Y, -HALF_WORLD_MAX), FVector(WorldViewRect.Max.X, WorldViewRect.Max.Y, HALF_WORLD_MAX));
 
 	TSet<FWorldPartitionActorDescViewBoundsProxy> ActorDescList;
 	
@@ -1089,8 +1089,8 @@ void SWorldPartitionEditorGrid2D::UpdateSelection()
 	if (SelectBox2D.GetArea() > 0.0f)
 	{
 		SelectBox = FBox(
-			FVector(SelectBox2D.Min.X, SelectBox2D.Min.Y, -WORLDPARTITION_MAX),
-			FVector(SelectBox2D.Max.X, SelectBox2D.Max.Y, WORLDPARTITION_MAX)
+			FVector(SelectBox2D.Min.X, SelectBox2D.Min.Y, -HALF_WORLD_MAX),
+			FVector(SelectBox2D.Max.X, SelectBox2D.Max.Y, HALF_WORLD_MAX)
 		);
 	}
 }
diff --git a/Engine/Source/Runtime/Engine/Private/WorldPartition/WorldPartition.cpp b/Engine/Source/Runtime/Engine/Private/WorldPartition/WorldPartition.cpp
--- a/Engine/Source/Runtime/Engine/Private/WorldPartition/WorldPartition.cpp
+++ b/Engine/Source/Runtime/Engine/Private/WorldPartition/WorldPartition.cpp
@@ -205,7 +205,7 @@ class FLoaderAdapterAlwaysLoadedActors : public FLoaderAdapterShape
 {
 public:
 	FLoaderAdapterAlwaysLoadedActors(UWorld* InWorld)
-		: FLoaderAdapterShape(InWorld, FBox(FVector(-WORLDPARTITION_MAX, -WORLDPARTITION_MAX, -WORLDPARTITION_MAX), FVector(WORLDPARTITION_MAX, WORLDPARTITION_MAX, WORLDPARTITION_MAX)), TEXT("Always Loaded"))
+		: FLoaderAdapterShape(InWorld, FBox(FVector(-HALF_WORLD_MAX, -HALF_WORLD_MAX, -HALF_WORLD_MAX), FVector(HALF_WORLD_MAX, HALF_WORLD_MAX, HALF_WORLD_MAX)), TEXT("Always Loaded"))
 	{
 		bIncludeSpatiallyLoadedActors = false;
 		bIncludeNonSpatiallyLoadedActors = true;
@@ -690,7 +690,7 @@ void UWorldPartition::OnPostBugItGoCalled(const FVector& Loc, const FRotator& Ro
 #if WITH_EDITOR
 	if (GetMutableDefault<UWorldPartitionEditorPerProjectUserSettings>()->GetBugItGoLoadRegion())
 	{
-		const FVector LoadExtent(GLoadingRangeBugItGo, GLoadingRangeBugItGo, WORLDPARTITION_MAX);
+		const FVector LoadExtent(GLoadingRangeBugItGo, GLoadingRangeBugItGo, HALF_WORLD_MAX);
 		const FBox LoadCellsBox(Loc - LoadExtent, Loc + LoadExtent);
 
 		FLoaderAdapterShape* LoaderAdapter = CreateEditorLoaderAdapter<FLoaderAdapterShape>(World, LoadCellsBox, TEXT("BugItGo"));
diff --git a/Engine/Source/Runtime/Engine/Private/WorldPartition/WorldPartitionMiniMapHelper.cpp b/Engine/Source/Runtime/Engine/Private/WorldPartition/WorldPartitionMiniMapHelper.cpp
--- a/Engine/Source/Runtime/Engine/Private/WorldPartition/WorldPartitionMiniMapHelper.cpp
+++ b/Engine/Source/Runtime/Engine/Private/WorldPartition/WorldPartitionMiniMapHelper.cpp
@@ -134,7 +134,7 @@ void FWorldPartitionMiniMapHelper::CalTopViewOfWorld(FMatrix& OutProjectionMatri
 	const float OrthoWidth = FMath::Max(1.f, ViewportWidth * OrthoZoom);
 	const float OrthoHeight = FMath::Max(1.f, ViewportHeight * OrthoZoom);
 
-	const float ZOffset = WORLDPARTITION_MAX * 0.5;
+	const float ZOffset = HALF_WORLD_MAX;
 	OutProjectionMatrix = FReversedZOrthoMatrix(
 		OrthoWidth,
 		OrthoHeight,
diff --git a/Engine/Source/Runtime/Engine/Private/WorldPartition/WorldPartitionRuntimeCell.cpp b/Engine/Source/Runtime/Engine/Private/WorldPartition/WorldPartitionRuntimeCell.cpp
--- a/Engine/Source/Runtime/Engine/Private/WorldPartition/WorldPartitionRuntimeCell.cpp
+++ b/Engine/Source/Runtime/Engine/Private/WorldPartition/WorldPartitionRuntimeCell.cpp
@@ -14,7 +14,7 @@ int32 UWorldPartitionRuntimeCell::StreamingSourceCacheEpoch = 0;
 UWorldPartitionRuntimeCell::UWorldPartitionRuntimeCell(const FObjectInitializer& ObjectInitializer)
 : Super(ObjectInitializer)
 , bIsAlwaysLoaded(false)
-, MinMaxZ(-WORLDPARTITION_MAX, WORLDPARTITION_MAX)
+, MinMaxZ(-HALF_WORLD_MAX, HALF_WORLD_MAX)
 , Priority(0)
 , CachedMinSourcePriority((uint8)EStreamingSourcePriority::Lowest)
 , CachedSourceInfoEpoch(INT_MIN)
diff --git a/Engine/Source/Runtime/Engine/Public/WorldPartition/WorldPartition.h b/Engine/Source/Runtime/Engine/Public/WorldPartition/WorldPartition.h
--- a/Engine/Source/Runtime/Engine/Public/WorldPartition/WorldPartition.h
+++ b/Engine/Source/Runtime/Engine/Public/WorldPartition/WorldPartition.h
@@ -48,8 +48,6 @@ enum class EWorldPartitionInitState
 	Uninitializing
 };
 
-#define WORLDPARTITION_MAX UE_OLD_WORLD_MAX
-
 #if WITH_EDITOR
 /**
  * Interface for the world partition editor
