World Partition minimap display
* Initially zoom to world bounds rather than editor world bounds

diff --git a/Engine/Source/Editor/WorldPartitionEditor/Private/WorldPartition/SWorldPartitionEditorGrid2D.cpp b/Engine/Source/Editor/WorldPartitionEditor/Private/WorldPartition/SWorldPartitionEditorGrid2D.cpp
--- a/Engine/Source/Editor/WorldPartitionEditor/Private/WorldPartition/SWorldPartitionEditorGrid2D.cpp
+++ b/Engine/Source/Editor/WorldPartitionEditor/Private/WorldPartition/SWorldPartitionEditorGrid2D.cpp
@@ -984,7 +984,7 @@ int32 SWorldPartitionEditorGrid2D::OnPaint(const FPaintArgs& Args, const FGeomet
 
 		if (bResetView)
 		{
-			const FBox WorldBounds = WorldPartition->GetEditorWorldBounds();
+			const FBox WorldBounds = WorldPartition->GetWorldBounds();
 			FocusBox(WorldBounds);
 		}
 
