WorldPartition - Editor LWC support

diff --git a/Engine/Source/Editor/WorldPartitionEditor/Private/WorldPartition/SWorldPartitionEditorGrid2D.cpp b/Engine/Source/Editor/WorldPartitionEditor/Private/WorldPartition/SWorldPartitionEditorGrid2D.cpp
--- a/Engine/Source/Editor/WorldPartitionEditor/Private/WorldPartition/SWorldPartitionEditorGrid2D.cpp
+++ b/Engine/Source/Editor/WorldPartitionEditor/Private/WorldPartition/SWorldPartitionEditorGrid2D.cpp
@@ -518,7 +518,7 @@ FReply SWorldPartitionEditorGrid2D::OnMouseWheel(const FGeometry& MyGeometry, co
 	FVector2D MousePosLocalSpace = MouseCursorPos - MyGeometry.GetLocalSize() * 0.5f;
 	FVector2D P0 = MousePosLocalSpace / Scale;
 	float Delta = 1.0f + FMath::Abs(MouseEvent.GetWheelDelta() / 4.0f);
-	Scale = FMath::Clamp(Scale * (MouseEvent.GetWheelDelta() > 0 ? Delta : (1.0f / Delta)), 0.0001f, 10.0f);
+	Scale = FMath::Clamp(Scale * (MouseEvent.GetWheelDelta() > 0 ? Delta : (1.0f / Delta)), 0.00000001f, 10.0f);
 	FVector2D P1 = MousePosLocalSpace / Scale;
 	Trans += (P1 - P0);
 	UpdateTransform();
@@ -1073,8 +1073,8 @@ void SWorldPartitionEditorGrid2D::FocusBox(const FBox& Box) const
 
 void SWorldPartitionEditorGrid2D::UpdateTransform() const
 {
-	FTransform2D T(1.0f, Trans);
-	FTransform2D V(Scale, FVector2D(ScreenRect.GetSize().X * 0.5f, ScreenRect.GetSize().Y * 0.5f));
+	FTransform2d T(1.0f, Trans);
+	FTransform2d V(Scale, FVector2D(ScreenRect.GetSize().X * 0.5f, ScreenRect.GetSize().Y * 0.5f));
 	WorldToScreen = T.Concatenate(V);
 	ScreenToWorld = WorldToScreen.Inverse();
 }
diff --git a/Engine/Source/Editor/WorldPartitionEditor/Private/WorldPartition/SWorldPartitionEditorGridSpatialHash.cpp b/Engine/Source/Editor/WorldPartitionEditor/Private/WorldPartition/SWorldPartitionEditorGridSpatialHash.cpp
--- a/Engine/Source/Editor/WorldPartitionEditor/Private/WorldPartition/SWorldPartitionEditorGridSpatialHash.cpp
+++ b/Engine/Source/Editor/WorldPartitionEditor/Private/WorldPartition/SWorldPartitionEditorGridSpatialHash.cpp
@@ -77,7 +77,7 @@ int32 SWorldPartitionEditorGridSpatialHash::PaintGrid(const FGeometry& AllottedG
 	UWorldPartitionEditorSpatialHash* EditorSpatialHash = (UWorldPartitionEditorSpatialHash*)WorldPartition->EditorHash;
 	
 	// Found the best cell size depending on the current zoom
-	int32 EffectiveCellSize = EditorSpatialHash->CellSize;
+	int64 EffectiveCellSize = EditorSpatialHash->CellSize;
 	for (const UWorldPartitionEditorSpatialHash::FCellNodeHashLevel& HashLevel : EditorSpatialHash->HashLevels)
 	{
 		const FVector2D CellScreenSize = WorldToScreen.TransformVector(FVector2D(EffectiveCellSize, EffectiveCellSize));
@@ -172,12 +172,12 @@ int32 SWorldPartitionEditorGridSpatialHash::PaintGrid(const FGeometry& AllottedG
 	{
 		const FLinearColor Color = FLinearColor(0.1f, 0.1f, 0.1f, 1.f);
 
-		FIntVector2 TopLeftW(
+		UE::Math::TIntVector2<int64> TopLeftW(
 			FMath::FloorToFloat(VisibleGridRectWorld.Min.X / EffectiveCellSize) * EffectiveCellSize,
 			FMath::FloorToFloat(VisibleGridRectWorld.Min.Y / EffectiveCellSize) * EffectiveCellSize
 		);
 
-		FIntVector2 BottomRightW(
+		UE::Math::TIntVector2<int64> BottomRightW(
 			FMath::CeilToFloat(VisibleGridRectWorld.Max.X / EffectiveCellSize) * EffectiveCellSize,
 			FMath::CeilToFloat(VisibleGridRectWorld.Max.Y / EffectiveCellSize) * EffectiveCellSize
 		);
@@ -186,7 +186,7 @@ int32 SWorldPartitionEditorGridSpatialHash::PaintGrid(const FGeometry& AllottedG
 		LinePoints.SetNum(2);
 
 		// Horizontal
-		for (int32 i=TopLeftW.Y; i<=BottomRightW.Y; i+=EffectiveCellSize)
+		for (int64 i=TopLeftW.Y; i<=BottomRightW.Y; i+=EffectiveCellSize)
 		{
 			FVector2D LineStartH(TopLeftW.X, i);
 			FVector2D LineEndH(BottomRightW.X, i);
@@ -198,7 +198,7 @@ int32 SWorldPartitionEditorGridSpatialHash::PaintGrid(const FGeometry& AllottedG
 		}
 
 		// Vertical
-		for (int32 i=TopLeftW.X; i<=BottomRightW.X; i+=EffectiveCellSize)
+		for (int64 i=TopLeftW.X; i<=BottomRightW.X; i+=EffectiveCellSize)
 		{
 			FVector2D LineStartH(i, TopLeftW.Y);
 			FVector2D LineEndH(i, BottomRightW.Y);
@@ -218,10 +218,12 @@ int32 SWorldPartitionEditorGridSpatialHash::PaintGrid(const FGeometry& AllottedG
 			FSlateFontInfo CoordsFont;
 			FVector2D DefaultCoordTextSize;
 			bool bNeedsGradient = true;
+			// Use top-left coordinate as default coord text
+			const FString DefaultCoordText(FString::Printf(TEXT("(%lld,%lld)"), TopLeftW.X / EffectiveCellSize, TopLeftW.Y / EffectiveCellSize));
 			for(int32 DesiredFontSize = 24; DesiredFontSize >= 8; DesiredFontSize -= 2)
 			{
 				CoordsFont = FCoreStyle::GetDefaultFontStyle("Bold", DesiredFontSize);
-				DefaultCoordTextSize = FontMeasure->Measure(TEXT("(-99,-99)"), CoordsFont);				
+				DefaultCoordTextSize = FontMeasure->Measure(DefaultCoordText, CoordsFont);
 
 				if (CellScreenSize.X > DefaultCoordTextSize.X)
 				{
@@ -236,18 +238,18 @@ int32 SWorldPartitionEditorGridSpatialHash::PaintGrid(const FGeometry& AllottedG
 				float ColorGradient = bNeedsGradient ? FMath::Min((CellScreenSize.X - DefaultCoordTextSize.X) / GradientDistance, 1.0f) : 1.0f;
 				const FLinearColor CoordTextColor(1.0f, 1.0f, 1.0f, ColorGradient);
 
-				for (int32 y = TopLeftW.Y; y < BottomRightW.Y; y += EffectiveCellSize)
+				for (int64 y = TopLeftW.Y; y < BottomRightW.Y; y += EffectiveCellSize)
 				{
-					for (int32 x = TopLeftW.X; x < BottomRightW.X; x += EffectiveCellSize)
+					for (int64 x = TopLeftW.X; x < BottomRightW.X; x += EffectiveCellSize)
 					{
-						const FString CoordText(FString::Printf(TEXT("(%d,%d)"), x / EffectiveCellSize, y / EffectiveCellSize));
+						const FString CoordText(FString::Printf(TEXT("(%lld,%lld)"), x / EffectiveCellSize, y / EffectiveCellSize));
 						const FVector2D CoordTextSize = FontMeasure->Measure(CoordText, CoordsFont);
 
 						FSlateDrawElement::MakeText(
 							OutDrawElements,
 							++LayerId,
 							AllottedGeometry.ToPaintGeometry(WorldToScreen.TransformPoint(FVector2D(x + EffectiveCellSize / 2, y + EffectiveCellSize / 2)) - CoordTextSize / 2, FVector2D(1,1)),
-							FString::Printf(TEXT("(%d,%d)"), x / EffectiveCellSize, y / EffectiveCellSize),
+							FString::Printf(TEXT("(%lld,%lld)"), x / EffectiveCellSize, y / EffectiveCellSize),
 							CoordsFont,
 							ESlateDrawEffect::None,
 							CoordTextColor
diff --git a/Engine/Source/Editor/WorldPartitionEditor/Public/WorldPartition/SWorldPartitionEditorGrid2D.h b/Engine/Source/Editor/WorldPartitionEditor/Public/WorldPartition/SWorldPartitionEditorGrid2D.h
--- a/Engine/Source/Editor/WorldPartitionEditor/Public/WorldPartition/SWorldPartitionEditorGrid2D.h
+++ b/Engine/Source/Editor/WorldPartitionEditor/Public/WorldPartition/SWorldPartitionEditorGrid2D.h
@@ -80,8 +80,8 @@ protected:
 	mutable FVector2D Trans;
 	
 	mutable FBox2D ScreenRect;
-	mutable FTransform2D WorldToScreen;
-	mutable FTransform2D ScreenToWorld;
+	mutable FTransform2d WorldToScreen;
+	mutable FTransform2d ScreenToWorld;
 
 	bool bIsSelecting;
 	bool bIsDragging;
diff --git a/Engine/Source/Runtime/Engine/Public/WorldPartition/WorldPartitionEditorSpatialHash.h b/Engine/Source/Runtime/Engine/Public/WorldPartition/WorldPartitionEditorSpatialHash.h
--- a/Engine/Source/Runtime/Engine/Public/WorldPartition/WorldPartitionEditorSpatialHash.h
+++ b/Engine/Source/Runtime/Engine/Public/WorldPartition/WorldPartitionEditorSpatialHash.h
@@ -81,7 +81,7 @@ class ENGINE_API UWorldPartitionEditorSpatialHash : public UWorldPartitionEditor
 	inline FCellCoord GetCellCoords(const FVector& InPos, int32 Level) const
 	{
 		check(Level >= 0);
-		const int32 CellSizeForLevel = CellSize * (1 << Level);
+		const int64 CellSizeForLevel = (int64)CellSize * (1LL << Level);
 		return FCellCoord(
 			FMath::FloorToInt(InPos.X / CellSizeForLevel),
 			FMath::FloorToInt(InPos.Y / CellSizeForLevel),
@@ -93,7 +93,7 @@ class ENGINE_API UWorldPartitionEditorSpatialHash : public UWorldPartitionEditor
 	inline FBox GetCellBounds(const FCellCoord& InCellCoord) const
 	{
 		check(InCellCoord.Level >= 0);
-		const int32 CellSizeForLevel = CellSize * (1 << InCellCoord.Level);
+		const int64 CellSizeForLevel = (int64)CellSize * (1LL << InCellCoord.Level);
 		const FVector Min = FVector(
 			InCellCoord.X * CellSizeForLevel, 
 			InCellCoord.Y * CellSizeForLevel, 
