Introduce WorldPartitionBlueprintLibrary: adds basic functionnality for Blutilities to query unloaded actors informations.

#rnx

diff --git a/Engine/Source/Editor/WorldPartitionEditor/Private/WorldPartition/WorldPartitionBlueprintLibrary.cpp b/Engine/Source/Editor/WorldPartitionEditor/Private/WorldPartition/WorldPartitionBlueprintLibrary.cpp
new file mode 100644
--- /dev/null
+++ b/Engine/Source/Editor/WorldPartitionEditor/Private/WorldPartition/WorldPartitionBlueprintLibrary.cpp
@@ -0,0 +1,148 @@
+// Copyright Epic Games, Inc. All Rights Reserved.
+
+#include "WorldPartition/WorldPartitionBlueprintLibrary.h"
+#include "Subsystems/UnrealEditorSubsystem.h"
+#include "WorldPartition/WorldPartition.h"
+#include "WorldPartition/WorldPartitionActorDesc.h"
+#include "WorldPartition/WorldPartitionHelpers.h"
+#include "Editor.h"
+
+FActorDesc::FActorDesc()
+	: Bounds(ForceInit)
+{}
+
+FActorDesc::FActorDesc(const FWorldPartitionActorDesc& InActorDesc, const FTransform& InTransform)
+{
+	Guid = InActorDesc.GetGuid();
+
+	if (!InActorDesc.GetBaseClass().IsNone())
+	{
+		Class = InActorDesc.GetBaseClass();
+	}
+	else
+	{
+		Class = InActorDesc.GetActorNativeClass();
+	}
+
+	Name = InActorDesc.GetActorName();
+	Label = InActorDesc.GetActorLabel();
+	Bounds = InActorDesc.GetBounds().TransformBy(InTransform);
+	RuntimeGrid = InActorDesc.GetRuntimeGrid();
+	bIsSpatiallyLoaded = InActorDesc.GetIsSpatiallyLoaded();
+	bActorIsEditorOnly = InActorDesc.GetActorIsEditorOnly();	
+}
+
+UWorld* UWorldPartitionBlueprintLibrary::GetEditorWorld()
+{
+	UUnrealEditorSubsystem* UnrealEditorSubsystem = GEditor->GetEditorSubsystem<UUnrealEditorSubsystem>();
+	return UnrealEditorSubsystem ? UnrealEditorSubsystem->GetEditorWorld() : nullptr;
+}
+
+UWorldPartition* UWorldPartitionBlueprintLibrary::GetWorldPartition()
+{
+	if (UWorld* World = GetEditorWorld())
+	{
+		return World->GetWorldPartition();
+	}
+	return nullptr;
+}
+
+bool UWorldPartitionBlueprintLibrary::GetActorDescs(const UActorDescContainer* InContainer, const FTransform& InTransform, TArray<FActorDesc>& OutActorDescs)
+{
+	bool bResult = true;
+
+	for (FActorDescList::TConstIterator<> ActorDescIt(InContainer); ActorDescIt; ++ActorDescIt)
+	{
+		const UActorDescContainer* SubContainer;
+		EContainerClusterMode SubClusterMode;
+		FTransform SubTransform;
+
+		if (ActorDescIt->IsContainerInstance())
+		{
+			if (ActorDescIt->GetContainerInstance(SubContainer, SubTransform, SubClusterMode))
+			{
+				bResult &= GetActorDescs(SubContainer, SubTransform * InTransform, OutActorDescs);
+			}
+			else
+			{
+				bResult = false;
+			}
+		}
+		else
+		{
+			OutActorDescs.Emplace(**ActorDescIt, InTransform);
+		}
+	}
+
+	return bResult;
+}
+
+bool UWorldPartitionBlueprintLibrary::GetIntersectingActorDescs(const UActorDescContainer* InContainer, const FBox& InBox, const FTransform& InTransform, TArray<FActorDesc>& OutActorDescs)
+{
+	bool bResult = true;
+
+	auto HandleActorDesc = [&InBox, &InTransform, &OutActorDescs](const FWorldPartitionActorDesc* ActorDesc)
+	{
+		bool bResult = true;
+		const UActorDescContainer* SubContainer;
+		EContainerClusterMode SubClusterMode;
+		FTransform SubTransform;
+
+		if (ActorDesc->IsContainerInstance())
+		{
+			if (ActorDesc->GetContainerInstance(SubContainer, SubTransform, SubClusterMode))
+			{
+				bResult &= GetIntersectingActorDescs(SubContainer, InBox, SubTransform * InTransform, OutActorDescs);
+			}
+			else
+			{
+				bResult = false;
+			}
+		}
+		else
+		{
+			OutActorDescs.Emplace(*ActorDesc, InTransform);
+		}
+
+		return bResult;
+	};
+
+	if (const UWorldPartition* WorldPartition = Cast<UWorldPartition>(InContainer))
+	{
+		FWorldPartitionHelpers::ForEachIntersectingActorDesc(const_cast<UWorldPartition*>(WorldPartition), InBox, [&bResult, &InBox, &InTransform, &OutActorDescs, &HandleActorDesc](const FWorldPartitionActorDesc* ActorDesc)
+		{
+			bResult &= HandleActorDesc(ActorDesc);
+			return true;
+		});
+	}
+	else
+	{
+		for (FActorDescList::TConstIterator<> ActorDescIt(InContainer); ActorDescIt; ++ActorDescIt)
+		{
+			if (ActorDescIt->GetBounds().Intersect(InBox))
+			{
+				bResult &= HandleActorDesc(*ActorDescIt);
+			}
+		}
+	}
+
+	return bResult;
+}
+
+bool UWorldPartitionBlueprintLibrary::GetActorDescs(TArray<FActorDesc>& OutActorDescs)
+{
+	if (UWorldPartition* WorldPartition = GetWorldPartition())
+	{
+		return GetActorDescs(WorldPartition, FTransform::Identity, OutActorDescs);
+	}
+	return false;
+}
+
+bool UWorldPartitionBlueprintLibrary::GetIntersectingActorDescs(const FBox& InBox, TArray<FActorDesc>& OutActorDescs)
+{
+	if (UWorldPartition* WorldPartition = GetWorldPartition())
+	{
+		return GetIntersectingActorDescs(WorldPartition, InBox, FTransform::Identity, OutActorDescs);
+	}
+	return false;
+}
diff --git a/Engine/Source/Editor/WorldPartitionEditor/Public/WorldPartition/WorldPartitionBlueprintLibrary.h b/Engine/Source/Editor/WorldPartitionEditor/Public/WorldPartition/WorldPartitionBlueprintLibrary.h
new file mode 100644
--- /dev/null
+++ b/Engine/Source/Editor/WorldPartitionEditor/Public/WorldPartition/WorldPartitionBlueprintLibrary.h
@@ -0,0 +1,82 @@
+// Copyright Epic Games, Inc. All Rights Reserved.
+
+#pragma once
+
+#include "Kismet/BlueprintFunctionLibrary.h"
+#include "WorldPartition/ActorDescContainer.h"
+#include "WorldPartition/WorldPartitionActorDesc.h"
+#include "WorldPartitionBlueprintLibrary.generated.h"
+
+class FWorldPartitionActorDesc;
+
+/**
+ * Snapshot of an actor descriptor, which represents the state of an actor on disk.
+ * The actor may or may not be loaded.
+ */
+USTRUCT(BlueprintType)
+struct FActorDesc
+{
+	GENERATED_USTRUCT_BODY()
+
+	FActorDesc();
+	FActorDesc(const FWorldPartitionActorDesc& InActorDesc, const FTransform& InTransform);
+
+	/** The actor GUID of this descriptor. */
+	UPROPERTY(BlueprintReadOnly, VisibleAnywhere, Category=Actor)
+	FGuid Guid;
+
+	/** Actor class, can point to a native or Blueprint class and may be redirected. */
+	UPROPERTY(BlueprintReadOnly, VisibleAnywhere, Category=Actor)
+	FSoftObjectPath Class;
+
+	/** Internal name of the acgor. */
+	UPROPERTY(BlueprintReadOnly, VisibleAnywhere, Category=Actor)
+	FName Name;
+
+	/** Actor's label. */
+	UPROPERTY(BlueprintReadOnly, VisibleAnywhere, Category=Actor)
+	FName Label;
+
+	/** Streaming bounds of this actor. */
+	UPROPERTY(BlueprintReadOnly, VisibleAnywhere, Category=Actor)
+	FBox Bounds;
+
+	/** Actor's target runtime grid. */
+	UPROPERTY(BlueprintReadOnly, VisibleAnywhere, Category=Actor)
+	FName RuntimeGrid;
+
+	/** Actor's streaming state. */
+	UPROPERTY(BlueprintReadOnly, VisibleAnywhere, Category=Actor)
+	bool bIsSpatiallyLoaded;
+
+	/** Actor's editor-only property. */
+	UPROPERTY(BlueprintReadOnly, VisibleAnywhere, Category=Actor)
+	bool bActorIsEditorOnly;
+};
+
+UCLASS(MinimalAPI)
+class UWorldPartitionBlueprintLibrary : public UBlueprintFunctionLibrary
+{
+	GENERATED_BODY()
+
+	static UWorld* GetEditorWorld();
+	static UWorldPartition* GetWorldPartition();
+
+	static bool GetActorDescs(const UActorDescContainer* InContainer, const FTransform& InTransform, TArray<FActorDesc>& OutActorDescs);
+	static bool GetIntersectingActorDescs(const UActorDescContainer* InContainer, const FBox& InBox, const FTransform& InTransform, TArray<FActorDesc>& OutActorDescs);
+
+public:
+	/**
+	 * Gets all the actor descriptors into the provided array, recursing into actor containers.
+	 * @return True if the operation was successful.
+	 */
+	UFUNCTION(BlueprintCallable, Category="World Partition")
+	static bool GetActorDescs(TArray<FActorDesc>& OutActorDescs);
+
+	/**
+	 * Gets all the actor descriptors intersecting the provided box into the provided array, recursing into actor containers.
+	 * @return True if the operation was successful.
+	 */
+	UFUNCTION(BlueprintCallable, Category="World Partition")
+	static bool GetIntersectingActorDescs(const FBox& InBox, TArray<FActorDesc>& OutActorDescs);
+};
\ No newline at end of file
