World Partition Editor QoL: allow unloading hovered region without making any selection.

#rnx

diff --git a/Engine/Source/Editor/WorldPartitionEditor/Private/WorldPartition/SWorldPartitionEditorGrid2D.cpp b/Engine/Source/Editor/WorldPartitionEditor/Private/WorldPartition/SWorldPartitionEditorGrid2D.cpp
--- a/Engine/Source/Editor/WorldPartitionEditor/Private/WorldPartition/SWorldPartitionEditorGrid2D.cpp
+++ b/Engine/Source/Editor/WorldPartitionEditor/Private/WorldPartition/SWorldPartitionEditorGrid2D.cpp
@@ -32,6 +32,44 @@ static bool IsBoundsSelected(const FBox& SelectBox, const FBox& Bounds)
 	return SelectBox.IsValid && Bounds.IntersectXY(SelectBox) && !Bounds.IsInsideXY(SelectBox);
 }
 
+static bool IsBoundsHovered(FVector2D Point, const FBox2D& Bounds)
+{
+	const float DistanceToPoint = FMath::Sqrt(Bounds.ComputeSquaredDistanceToPoint(Point));
+	return DistanceToPoint > UE_KINDA_SMALL_NUMBER && DistanceToPoint < 10.0f;
+}
+
+static bool IsActorBoundsHovered(const FBox& ActorBounds, const FTransform2d& WorldToScreen, const FVector2D& MouseCursorPos)
+{
+	FVector Origin, Extent;
+	ActorBounds.GetCenterAndExtents(Origin, Extent);
+
+	const FVector2D TopLeftW = FVector2D(Origin - Extent);
+	const FVector2D BottomRightW = FVector2D(Origin + Extent);
+	const FVector2D TopLeft = WorldToScreen.TransformPoint(TopLeftW);
+	const FVector2D BottomRight = WorldToScreen.TransformPoint(BottomRightW);
+	const FBox2D ActorViewBox(TopLeft, BottomRight);
+
+	return IsBoundsHovered(MouseCursorPos, ActorViewBox);
+}
+
+template <class T>
+void ForEachHoveredLoaderAdapters(const UWorldPartition* WorldPartition, const FTransform2d& WorldToScreen, const FVector2D& MouseCursorPos, T Func)
+{
+	for (IWorldPartitionActorLoaderInterface::ILoaderAdapter* LoaderAdapter : WorldPartition->GetRegisteredEditorLoaderAdapters())
+	{
+		if (LoaderAdapter->GetBoundingBox().IsSet())
+		{
+			if (IsActorBoundsHovered(*LoaderAdapter->GetBoundingBox(), WorldToScreen, MouseCursorPos))
+			{
+				if (!Func(LoaderAdapter))
+				{
+					return;
+				}
+			}
+		}
+	}
+}
+
 template <class T>
 void ForEachSelectedEditorLoaderAdapters(const UWorldPartition* WorldPartition, const FBox& SelectBox, T Func)
 {
@@ -119,6 +157,7 @@ void SWorldPartitionEditorGrid2D::FEditorCommands::RegisterCommands()
 	UI_COMMAND(CreateRegionFromSelection, "Create Loading Region From Selection", "Create a loading region from the selection.", EUserInterfaceActionType::Button, FInputChord());
 	UI_COMMAND(LoadSelectedRegions, "Load Selected Regions", "Load the selected regions.", EUserInterfaceActionType::Button, FInputChord());
 	UI_COMMAND(UnloadSelectedRegions, "Unload Selected Regions", "Unload the selected regions.", EUserInterfaceActionType::Button, FInputChord());
+	UI_COMMAND(UnloadHoveredRegion, "Unload Hovered Region", "Unload the hovered region.", EUserInterfaceActionType::Button, FInputChord());
 	UI_COMMAND(ConvertSelectedRegionsToActors, "Convert Selected Regions To Actors", "Convert the selected regions to actors.", EUserInterfaceActionType::Button, FInputChord());
 	UI_COMMAND(MoveCameraHere, "Move Camera Here", "Move the camera to the selected position.", EUserInterfaceActionType::Button, FInputChord());
 }
@@ -272,10 +311,26 @@ void SWorldPartitionEditorGrid2D::Construct(const FArguments& InArgs)
 
 	auto CanLoadSelectedRegions = [this, CanLoadUnloadSelectedRegions]() { return CanLoadUnloadSelectedRegions(true); };
 	auto CanUnloadSelectedRegions = [this, CanLoadUnloadSelectedRegions]() { return CanLoadUnloadSelectedRegions(false); };
+	
+	auto CanUnloadHoveredRegion = [this]()
+	{
+		bool bResult = false;
+		ForEachHoveredLoaderAdapters(WorldPartition, WorldToScreen, MouseCursorPos, [&bResult](IWorldPartitionActorLoaderInterface::ILoaderAdapter* LoaderAdapter)
+		{
+			if (LoaderAdapter->IsLoaded())
+			{
+				bResult = true;
+				return false;
+			}
+			return true;
+		});
+		return bResult;
+	};
 
 	ActionList.MapAction(Commands.CreateRegionFromSelection, FExecuteAction::CreateSP(this, &SWorldPartitionEditorGrid2D::CreateRegionFromSelection), FCanExecuteAction::CreateLambda(CanCreateRegionFromSelection));
 	ActionList.MapAction(Commands.LoadSelectedRegions, FExecuteAction::CreateSP(this, &SWorldPartitionEditorGrid2D::LoadSelectedRegions), FCanExecuteAction::CreateLambda(CanLoadSelectedRegions));
 	ActionList.MapAction(Commands.UnloadSelectedRegions, FExecuteAction::CreateSP(this, &SWorldPartitionEditorGrid2D::UnloadSelectedRegions), FCanExecuteAction::CreateLambda(CanUnloadSelectedRegions));
+	ActionList.MapAction(Commands.UnloadHoveredRegion, FExecuteAction::CreateSP(this, &SWorldPartitionEditorGrid2D::UnloadHoveredRegion), FCanExecuteAction::CreateLambda(CanUnloadHoveredRegion));
 	ActionList.MapAction(Commands.ConvertSelectedRegionsToActors, FExecuteAction::CreateSP(this, &SWorldPartitionEditorGrid2D::ConvertSelectedRegionsToActors), FCanExecuteAction::CreateLambda(CanConvertSelectedRegionsToActors));
 	ActionList.MapAction(Commands.MoveCameraHere, FExecuteAction::CreateSP(this, &SWorldPartitionEditorGrid2D::MoveCameraHere));
 }
@@ -338,6 +393,27 @@ void SWorldPartitionEditorGrid2D::UnloadSelectedRegions()
 	Refresh();
 }
 
+void SWorldPartitionEditorGrid2D::UnloadHoveredRegion()
+{
+	TArray<IWorldPartitionActorLoaderInterface::ILoaderAdapter*> LoaderAdaptersToRelease;
+	ForEachHoveredLoaderAdapters(WorldPartition, WorldToScreen, MouseCursorPos, [&LoaderAdaptersToRelease](IWorldPartitionActorLoaderInterface::ILoaderAdapter* LoaderAdapter)
+	{
+		if (LoaderAdapter->IsLoaded())
+		{
+			LoaderAdaptersToRelease.Add(LoaderAdapter);
+		}
+		return true;
+	});
+
+	for (IWorldPartitionActorLoaderInterface::ILoaderAdapter* LoaderAdapterToRelease : LoaderAdaptersToRelease)
+	{
+		WorldPartition->ReleaseEditorLoaderAdapter(LoaderAdapterToRelease);
+	}
+
+	GEditor->RedrawLevelEditingViewports();
+	Refresh();
+}
+
 void SWorldPartitionEditorGrid2D::ConvertSelectedRegionsToActors()
 {
 	TArray<IWorldPartitionActorLoaderInterface::ILoaderAdapter*> LoaderAdaptersToRelease;
@@ -432,6 +508,7 @@ FReply SWorldPartitionEditorGrid2D::OnMouseButtonUp(const FGeometry& MyGeometry,
 				MenuBuilder.AddMenuSeparator();
 				MenuBuilder.AddMenuEntry(Commands.LoadSelectedRegions);
 				MenuBuilder.AddMenuEntry(Commands.UnloadSelectedRegions);
+				MenuBuilder.AddMenuEntry(Commands.UnloadHoveredRegion);
 				MenuBuilder.AddMenuSeparator();
 				MenuBuilder.AddMenuEntry(Commands.ConvertSelectedRegionsToActors);
 			MenuBuilder.EndSection();
@@ -704,7 +781,7 @@ uint32 SWorldPartitionEditorGrid2D::PaintActors(const FGeometry& AllottedGeometr
 						{
 							const FLinearColor LoaderColor = LoaderAdapter->GetColor().IsSet() ? *LoaderAdapter->GetColor() : FColor::White;
 							const FLinearColor OutlineColor(LoaderColor.R, LoaderColor.G, LoaderColor.B, LoaderColorGradient * FullScreenColorGradient);
-							const bool IsHighlighted = IsBoundsSelected(SelectBox, ActorBounds);
+							const bool IsHighlighted = IsBoundsSelected(SelectBox, ActorBounds) || IsBoundsHovered(MouseCursorPos, ActorViewBox);
 
 							LinePoints[0] = TopLeft;
 							LinePoints[1] = TopRight;
@@ -776,8 +853,7 @@ uint32 SWorldPartitionEditorGrid2D::PaintActors(const FGeometry& AllottedGeometr
 				const AActor* Actor = ActorDescView.GetActor();
 
 				const bool bIsSelected = Actor ? Actor->IsSelected() : false;
-				const float SquaredDistanceToPoint = ActorViewBox.ComputeSquaredDistanceToPoint(MouseCursorPos);
-				if (bIsSelected || (SquaredDistanceToPoint > 0.0f && SquaredDistanceToPoint <= 2.0f))
+				if (bIsSelected || IsBoundsHovered(MouseCursorPos, ActorViewBox))
 				{
 					ActorColor = FLinearColor::Yellow;
 
diff --git a/Engine/Source/Editor/WorldPartitionEditor/Public/WorldPartition/SWorldPartitionEditorGrid2D.h b/Engine/Source/Editor/WorldPartitionEditor/Public/WorldPartition/SWorldPartitionEditorGrid2D.h
--- a/Engine/Source/Editor/WorldPartitionEditor/Public/WorldPartition/SWorldPartitionEditorGrid2D.h
+++ b/Engine/Source/Editor/WorldPartitionEditor/Public/WorldPartition/SWorldPartitionEditorGrid2D.h
@@ -18,6 +18,7 @@ protected:
 		TSharedPtr<FUICommandInfo> CreateRegionFromSelection;
 		TSharedPtr<FUICommandInfo> LoadSelectedRegions;
 		TSharedPtr<FUICommandInfo> UnloadSelectedRegions;
+		TSharedPtr<FUICommandInfo> UnloadHoveredRegion;
 		TSharedPtr<FUICommandInfo> ConvertSelectedRegionsToActors;
 		TSharedPtr<FUICommandInfo> MoveCameraHere;
 
@@ -36,6 +37,7 @@ public:
 	virtual void CreateRegionFromSelection();
 	virtual void LoadSelectedRegions();
 	virtual void UnloadSelectedRegions();
+	virtual void UnloadHoveredRegion();
 	virtual void ConvertSelectedRegionsToActors();
 	void MoveCameraHere();
 
