World Partition: Editor improvements
- Selection: drag select, click select, ctrl+select to add/remove from selection
- Right-click on region allows to run action on it (load, unload, etc)
- Hover highlight is now based on closest edge distance

diff --git a/Engine/Source/Editor/WorldPartitionEditor/Private/WorldPartition/SWorldPartitionEditorGrid2D.cpp b/Engine/Source/Editor/WorldPartitionEditor/Private/WorldPartition/SWorldPartitionEditorGrid2D.cpp
--- a/Engine/Source/Editor/WorldPartitionEditor/Private/WorldPartition/SWorldPartitionEditorGrid2D.cpp
+++ b/Engine/Source/Editor/WorldPartitionEditor/Private/WorldPartition/SWorldPartitionEditorGrid2D.cpp
@@ -6,11 +6,12 @@
 #include "GameFramework/WorldSettings.h"
 #include "Modules/ModuleManager.h"
 #include "WorldPartition/WorldPartition.h"
+#include "WorldPartition/WorldPartitionEditorLoaderAdapter.h"
 #include "WorldPartition/WorldPartitionEditorHash.h"
 #include "WorldPartition/WorldPartitionActorDesc.h"
 #include "WorldPartition/WorldPartitionActorDescView.h"
-#include "WorldPartition/WorldPartitionActorLoaderInterface.h"
 #include "WorldPartition/WorldPartitionEditorPerProjectUserSettings.h"
+#include "WorldPartition/WorldPartitionHelpers.h"
 #include "WorldPartition/LoaderAdapter/LoaderAdapterShape.h"
 #include "ActorFactories/ActorFactory.h"
 #include "Builders/CubeBuilder.h"
@@ -32,66 +33,38 @@ static bool IsBoundsSelected(const FBox& SelectBox, const FBox& Bounds)
 	return SelectBox.IsValid && Bounds.IntersectXY(SelectBox) && !Bounds.IsInsideXY(SelectBox);
 }
 
-static bool IsBoundsHovered(FVector2D Point, const FBox2D& Bounds)
+static bool IsBoundsHoveredXY(FVector2D Point, const FBox& Bounds, double& OutDistanceToEdge)
 {
-	const float DistanceToPoint = FMath::Sqrt(Bounds.ComputeSquaredDistanceToPoint(Point));
-	return DistanceToPoint > UE_KINDA_SMALL_NUMBER && DistanceToPoint < 10.0f;
-}
-
-static bool IsActorBoundsHovered(const FBox& ActorBounds, const FTransform2d& WorldToScreen, const FVector2D& MouseCursorPos)
-{
-	FVector Origin, Extent;
-	ActorBounds.GetCenterAndExtents(Origin, Extent);
-
-	const FVector2D TopLeftW = FVector2D(Origin - Extent);
-	const FVector2D BottomRightW = FVector2D(Origin + Extent);
-	const FVector2D TopLeft = WorldToScreen.TransformPoint(TopLeftW);
-	const FVector2D BottomRight = WorldToScreen.TransformPoint(BottomRightW);
-	const FBox2D ActorViewBox(TopLeft, BottomRight);
-
-	return IsBoundsHovered(MouseCursorPos, ActorViewBox);
-}
-
-template <class T>
-void ForEachHoveredLoaderAdapters(const UWorldPartition* WorldPartition, const FTransform2d& WorldToScreen, const FVector2D& MouseCursorPos, T Func)
-{
-	for (IWorldPartitionActorLoaderInterface::ILoaderAdapter* LoaderAdapter : WorldPartition->GetRegisteredEditorLoaderAdapters())
+	const FBox2D Bounds2D(FVector2D(Bounds.Min.X, Bounds.Min.Y), FVector2D(Bounds.Max.X, Bounds.Max.Y));
+	if (Bounds2D.IsInsideOrOn(Point))
 	{
-		if (LoaderAdapter->GetBoundingBox().IsSet())
-		{
-			if (IsActorBoundsHovered(*LoaderAdapter->GetBoundingBox(), WorldToScreen, MouseCursorPos))
-			{
-				if (!Func(LoaderAdapter))
-				{
-					return;
-				}
-			}
-		}
+		const double DistX = FMath::Min(Point.X - Bounds.Min.X, Bounds.Max.X - Point.X);
+		const double DistY = FMath::Min(Point.Y - Bounds.Min.Y, Bounds.Max.Y - Point.Y);
+		OutDistanceToEdge = FMath::Min(DistX, DistY);
+		return true;
 	}
+
+	return false;
 }
 
 template <class T>
-void ForEachSelectedEditorLoaderAdapters(const UWorldPartition* WorldPartition, const FBox& SelectBox, T Func)
+void ForEachIntersectingLoaderAdapters(UWorldPartition* WorldPartition, const FBox& SelectBox, T Func)
 {
-	for (IWorldPartitionActorLoaderInterface::ILoaderAdapter* LoaderAdapter : WorldPartition->GetRegisteredEditorLoaderAdapters())
+	for (UWorldPartitionEditorLoaderAdapter* EditorLoaderAdapter : WorldPartition->GetRegisteredEditorLoaderAdapters())
 	{
-		if (LoaderAdapter->GetBoundingBox().IsSet())
+		if (IWorldPartitionActorLoaderInterface::ILoaderAdapter* LoaderAdapter = EditorLoaderAdapter->GetLoaderAdapter())
 		{
-			if (IsBoundsSelected(SelectBox, *LoaderAdapter->GetBoundingBox()))
+			if (LoaderAdapter->GetBoundingBox().IsSet() && IsBoundsSelected(SelectBox, *LoaderAdapter->GetBoundingBox()))
 			{
-				if (!Func(LoaderAdapter))
+				if (!Func(EditorLoaderAdapter))
 				{
 					return;
 				}
 			}
 		}
 	}
-};
 
-template <class T>
-void ForEachSelectedWorldLoaderAdapters(const UWorldPartition* WorldPartition, const FBox& SelectBox, T Func)
-{
-	WorldPartition->EditorHash->ForEachIntersectingActor(SelectBox, [&](FWorldPartitionActorDesc* ActorDesc)
+	FWorldPartitionHelpers::ForEachIntersectingActorDesc(WorldPartition, SelectBox, [&](const FWorldPartitionActorDesc* ActorDesc)
 	{
 		if (AActor* Actor = ActorDesc->GetActor())
 		{
@@ -101,22 +74,17 @@ void ForEachSelectedWorldLoaderAdapters(const UWorldPartition* WorldPartition, c
 				{
 					if (IsBoundsSelected(SelectBox, ActorDesc->GetBounds()))
 					{
-						if (!Func(LoaderAdapter))
+						if (!Func(Actor))
 						{
-							return;
+							return false;
 						}
 					}
 				}
 			}
 		}
-	});
-}
 
-template <class T>
-void ForEachSelectedLoaderAdapters(const UWorldPartition* WorldPartition, const FBox& SelectBox, T Func)
-{
-	ForEachSelectedEditorLoaderAdapters<T>(WorldPartition, SelectBox, Func);
-	ForEachSelectedWorldLoaderAdapters<T>(WorldPartition, SelectBox, Func);
+		return true;
+	});
 }
 
 class FWorldPartitionActorDescViewBoundsProxy : public FWorldPartitionActorDescView
@@ -157,7 +125,6 @@ void SWorldPartitionEditorGrid2D::FEditorCommands::RegisterCommands()
 	UI_COMMAND(CreateRegionFromSelection, "Create Loading Region From Selection", "Create a loading region from the selection.", EUserInterfaceActionType::Button, FInputChord());
 	UI_COMMAND(LoadSelectedRegions, "Load Selected Regions", "Load the selected regions.", EUserInterfaceActionType::Button, FInputChord());
 	UI_COMMAND(UnloadSelectedRegions, "Unload Selected Regions", "Unload the selected regions.", EUserInterfaceActionType::Button, FInputChord());
-	UI_COMMAND(UnloadHoveredRegion, "Unload Hovered Region", "Unload the hovered region.", EUserInterfaceActionType::Button, FInputChord());
 	UI_COMMAND(ConvertSelectedRegionsToActors, "Convert Selected Regions To Actors", "Convert the selected regions to actors.", EUserInterfaceActionType::Button, FInputChord());
 	UI_COMMAND(MoveCameraHere, "Move Camera Here", "Move the camera to the selected position.", EUserInterfaceActionType::Button, FInputChord());
 }
@@ -168,8 +135,8 @@ SWorldPartitionEditorGrid2D::SWorldPartitionEditorGrid2D()
 	, Scale(0.001f)
 	, Trans(ForceInit)
 	, ScreenRect(ForceInit)
-	, bIsSelecting(false)
-	, bIsDragging(false)
+	, bIsDragSelecting(false)
+	, bIsPanning(false)
 	, bShowActors(false)
 	, SelectBox(ForceInit)
 {
@@ -285,52 +252,38 @@ void SWorldPartitionEditorGrid2D::Construct(const FArguments& InArgs)
 
 	auto CanLoadUnloadSelectedRegions = [this](bool bLoad)
 	{
-		bool bResult = false;
-		ForEachSelectedLoaderAdapters(WorldPartition, SelectBox, [&bResult, bLoad](IWorldPartitionActorLoaderInterface::ILoaderAdapter* LoaderAdapter)
+		for (const TWeakInterfacePtr<IWorldPartitionActorLoaderInterface>& SelectedLoaderAdapter : SelectedLoaderAdapters)
 		{
-			if (bLoad != LoaderAdapter->IsLoaded())
+			if (IWorldPartitionActorLoaderInterface* LoaderInterface = SelectedLoaderAdapter.Get())
 			{
-				bResult = true;
-				return false;
+				if (bLoad != LoaderInterface->GetLoaderAdapter()->IsLoaded())
+				{
+					return true;
+				}
 			}
-			return true;
-		});
-		return bResult;
-	};
+		}
 
-	auto CanConvertSelectedRegionsToActors = [this]()
-	{
-		bool bResult = false;
-		ForEachSelectedEditorLoaderAdapters(WorldPartition, SelectBox, [&bResult](IWorldPartitionActorLoaderInterface::ILoaderAdapter* LoaderAdapter)
-		{
-			bResult = true;
-			return false;
-		});
-		return bResult;
+		return false;
 	};
 
-	auto CanLoadSelectedRegions = [this, CanLoadUnloadSelectedRegions]() { return CanLoadUnloadSelectedRegions(true); };
-	auto CanUnloadSelectedRegions = [this, CanLoadUnloadSelectedRegions]() { return CanLoadUnloadSelectedRegions(false); };
-	
-	auto CanUnloadHoveredRegion = [this]()
+	auto CanConvertSelectedRegionsToActors = [this]()
 	{
-		bool bResult = false;
-		ForEachHoveredLoaderAdapters(WorldPartition, WorldToScreen, MouseCursorPos, [&bResult](IWorldPartitionActorLoaderInterface::ILoaderAdapter* LoaderAdapter)
+		for (const TWeakInterfacePtr<IWorldPartitionActorLoaderInterface>& SelectedLoaderAdapter : SelectedLoaderAdapters)
 		{
-			if (LoaderAdapter->IsLoaded())
+			if (UWorldPartitionEditorLoaderAdapter* EditorLoaderAdapter = Cast<UWorldPartitionEditorLoaderAdapter>(SelectedLoaderAdapter.Get()))
 			{
-				bResult = true;
-				return false;
+				return true;
 			}
-			return true;
-		});
-		return bResult;
+		};
+		return false;
 	};
 
+	auto CanLoadSelectedRegions = [this, CanLoadUnloadSelectedRegions]() { return CanLoadUnloadSelectedRegions(true); };
+	auto CanUnloadSelectedRegions = [this, CanLoadUnloadSelectedRegions]() { return CanLoadUnloadSelectedRegions(false); };
+		
 	ActionList.MapAction(Commands.CreateRegionFromSelection, FExecuteAction::CreateSP(this, &SWorldPartitionEditorGrid2D::CreateRegionFromSelection), FCanExecuteAction::CreateLambda(CanCreateRegionFromSelection));
 	ActionList.MapAction(Commands.LoadSelectedRegions, FExecuteAction::CreateSP(this, &SWorldPartitionEditorGrid2D::LoadSelectedRegions), FCanExecuteAction::CreateLambda(CanLoadSelectedRegions));
 	ActionList.MapAction(Commands.UnloadSelectedRegions, FExecuteAction::CreateSP(this, &SWorldPartitionEditorGrid2D::UnloadSelectedRegions), FCanExecuteAction::CreateLambda(CanUnloadSelectedRegions));
-	ActionList.MapAction(Commands.UnloadHoveredRegion, FExecuteAction::CreateSP(this, &SWorldPartitionEditorGrid2D::UnloadHoveredRegion), FCanExecuteAction::CreateLambda(CanUnloadHoveredRegion));
 	ActionList.MapAction(Commands.ConvertSelectedRegionsToActors, FExecuteAction::CreateSP(this, &SWorldPartitionEditorGrid2D::ConvertSelectedRegionsToActors), FCanExecuteAction::CreateLambda(CanConvertSelectedRegionsToActors));
 	ActionList.MapAction(Commands.MoveCameraHere, FExecuteAction::CreateSP(this, &SWorldPartitionEditorGrid2D::MoveCameraHere));
 }
@@ -338,11 +291,12 @@ void SWorldPartitionEditorGrid2D::Construct(const FArguments& InArgs)
 void SWorldPartitionEditorGrid2D::CreateRegionFromSelection()
 {
 	const FBox RegionBox(FVector(SelectBox.Min.X, SelectBox.Min.Y, -HALF_WORLD_MAX), FVector(SelectBox.Max.X, SelectBox.Max.Y, HALF_WORLD_MAX));
-	FLoaderAdapterShape* LoaderAdapter = WorldPartition->CreateEditorLoaderAdapter<FLoaderAdapterShape>(World, RegionBox, TEXT("Loaded Region"));
-	LoaderAdapter->SetUserCreated(true);
-	LoaderAdapter->Load();
-
-	SelectBox.Init();
+	UWorldPartitionEditorLoaderAdapter* EditorLoaderAdapter = WorldPartition->CreateEditorLoaderAdapter<FLoaderAdapterShape>(World, RegionBox, TEXT("Loaded Region"));
+	EditorLoaderAdapter->GetLoaderAdapter()->SetUserCreated(true);
+	EditorLoaderAdapter->GetLoaderAdapter()->Load();
+	
+	ClearSelection();
+	SelectedLoaderAdapters.Add(EditorLoaderAdapter);
 
 	GEditor->RedrawLevelEditingViewports();
 	Refresh();
@@ -350,65 +304,38 @@ void SWorldPartitionEditorGrid2D::CreateRegionFromSelection()
 
 void SWorldPartitionEditorGrid2D::LoadSelectedRegions()
 {
-	check(SelectBox.IsValid);
-
-	ForEachSelectedLoaderAdapters(WorldPartition, SelectBox, [](IWorldPartitionActorLoaderInterface::ILoaderAdapter* LoaderAdapter)
-	{
-		LoaderAdapter->Load();
-		return true;
-	});
-
-	SelectBox.Init();
-
-	GEditor->RedrawLevelEditingViewports();
-	Refresh();
-}
-
-void SWorldPartitionEditorGrid2D::UnloadSelectedRegions()
-{
-	check(SelectBox.IsValid);
-
-	TArray<IWorldPartitionActorLoaderInterface::ILoaderAdapter*> LoaderAdaptersToRelease;
-	const TSet<IWorldPartitionActorLoaderInterface::ILoaderAdapter*>& RegisteredLoaderAdapters = WorldPartition->GetRegisteredEditorLoaderAdapters();
-
-	ForEachSelectedLoaderAdapters(WorldPartition, SelectBox, [&LoaderAdaptersToRelease, &RegisteredLoaderAdapters](IWorldPartitionActorLoaderInterface::ILoaderAdapter* LoaderAdapter)
+	for (const TWeakInterfacePtr<IWorldPartitionActorLoaderInterface>& SelectedLoaderAdapter : SelectedLoaderAdapters)
 	{
-		LoaderAdapter->Unload();
-
-		if (RegisteredLoaderAdapters.Contains(LoaderAdapter))
+		if (IWorldPartitionActorLoaderInterface* LoaderInterface = SelectedLoaderAdapter.Get())
 		{
-			LoaderAdaptersToRelease.Add(LoaderAdapter);
+			LoaderInterface->GetLoaderAdapter()->Load();
 		}
-		return true;
-	});
-
-	for (IWorldPartitionActorLoaderInterface::ILoaderAdapter* LoaderAdapterToRelease : LoaderAdaptersToRelease)
-	{
-		WorldPartition->ReleaseEditorLoaderAdapter(LoaderAdapterToRelease);
 	}
 
 	SelectBox.Init();
-
+		
 	GEditor->RedrawLevelEditingViewports();
 	Refresh();
 }
 
-void SWorldPartitionEditorGrid2D::UnloadHoveredRegion()
+void SWorldPartitionEditorGrid2D::UnloadSelectedRegions()
 {
-	TArray<IWorldPartitionActorLoaderInterface::ILoaderAdapter*> LoaderAdaptersToRelease;
-	ForEachHoveredLoaderAdapters(WorldPartition, WorldToScreen, MouseCursorPos, [&LoaderAdaptersToRelease](IWorldPartitionActorLoaderInterface::ILoaderAdapter* LoaderAdapter)
+	FLoaderAdapterSet CopySelectedLoaderAdapters(SelectedLoaderAdapters);
+	for (const TWeakInterfacePtr<IWorldPartitionActorLoaderInterface>& SelectedLoaderAdapter : CopySelectedLoaderAdapters)
 	{
-		if (LoaderAdapter->IsLoaded())
+		if (IWorldPartitionActorLoaderInterface* LoaderInterface = SelectedLoaderAdapter.Get())
 		{
-			LoaderAdaptersToRelease.Add(LoaderAdapter);
-		}
-		return true;
-	});
+			LoaderInterface->GetLoaderAdapter()->Unload();
 
-	for (IWorldPartitionActorLoaderInterface::ILoaderAdapter* LoaderAdapterToRelease : LoaderAdaptersToRelease)
-	{
-		WorldPartition->ReleaseEditorLoaderAdapter(LoaderAdapterToRelease);
-	}
+			if (UWorldPartitionEditorLoaderAdapter* EditorLoaderAdapter = Cast<UWorldPartitionEditorLoaderAdapter>(LoaderInterface))
+			{
+				SelectedLoaderAdapters.Remove(EditorLoaderAdapter);
+				WorldPartition->ReleaseEditorLoaderAdapter(EditorLoaderAdapter);
+			}
+		}
+	};
+		
+	SelectBox.Init();
 
 	GEditor->RedrawLevelEditingViewports();
 	Refresh();
@@ -416,33 +343,41 @@ void SWorldPartitionEditorGrid2D::UnloadHoveredRegion()
 
 void SWorldPartitionEditorGrid2D::ConvertSelectedRegionsToActors()
 {
-	TArray<IWorldPartitionActorLoaderInterface::ILoaderAdapter*> LoaderAdaptersToRelease;
-	ForEachSelectedEditorLoaderAdapters(WorldPartition, SelectBox, [this, &LoaderAdaptersToRelease](IWorldPartitionActorLoaderInterface::ILoaderAdapter* LoaderAdapter)
+	FLoaderAdapterSet TmpSelectedLoaderAdapters = MoveTemp(SelectedLoaderAdapters);
+	ClearSelection();
+
+	for (const TWeakInterfacePtr<IWorldPartitionActorLoaderInterface>& SelectedLoaderAdapter : TmpSelectedLoaderAdapters)
 	{
-		const FBox LoaderVolumeBox(*LoaderAdapter->GetBoundingBox());
+		if (UWorldPartitionEditorLoaderAdapter* EditorLoaderAdapter = Cast<UWorldPartitionEditorLoaderAdapter>(SelectedLoaderAdapter.Get()))
+		{
+			IWorldPartitionActorLoaderInterface::ILoaderAdapter* LoaderAdapter = EditorLoaderAdapter->GetLoaderAdapter();
+			
+			const FBox LoaderVolumeBox(*LoaderAdapter->GetBoundingBox());
 
-		ALocationVolume* LocationVolume = World->SpawnActor<ALocationVolume>(LoaderVolumeBox.GetCenter(), FRotator::ZeroRotator);
+			ALocationVolume* LocationVolume = World->SpawnActor<ALocationVolume>(LoaderVolumeBox.GetCenter(), FRotator::ZeroRotator);
 
-		UCubeBuilder* Builder = NewObject<UCubeBuilder>();
-		Builder->X = 1.0f;
-		Builder->Y = 1.0f;
-		Builder->Z = 1.0f;
-		UActorFactory::CreateBrushForVolumeActor(LocationVolume, Builder);
+			UCubeBuilder* Builder = NewObject<UCubeBuilder>();
+			Builder->X = 1.0f;
+			Builder->Y = 1.0f;
+			Builder->Z = 1.0f;
+			UActorFactory::CreateBrushForVolumeActor(LocationVolume, Builder);
 
-		LocationVolume->GetRootComponent()->SetWorldScale3D(LoaderVolumeBox.GetSize());
+			LocationVolume->GetRootComponent()->SetWorldScale3D(LoaderVolumeBox.GetSize());
 
-		LocationVolume->GetLoaderAdapter()->Load();
-		
-		LoaderAdapter->Unload();
-		LoaderAdaptersToRelease.Add(LoaderAdapter);
-
-		return true;
-	});
+			LocationVolume->GetLoaderAdapter()->Load();
 
-	for (IWorldPartitionActorLoaderInterface::ILoaderAdapter* LoaderAdapterToRelease : LoaderAdaptersToRelease)
-	{
-		WorldPartition->ReleaseEditorLoaderAdapter(LoaderAdapterToRelease);
+			LoaderAdapter->Unload();
+			
+			WorldPartition->ReleaseEditorLoaderAdapter(EditorLoaderAdapter);
+		}
+		else
+		{
+			SelectedLoaderAdapters.Add(SelectedLoaderAdapter);
+		}
 	}
+
+	GEditor->RedrawLevelEditingViewports();
+	Refresh();
 }
 
 void SWorldPartitionEditorGrid2D::MoveCameraHere()
@@ -472,7 +407,6 @@ FReply SWorldPartitionEditorGrid2D::OnMouseButtonDown(const FGeometry& MyGeometr
 
 		if (bIsLeftMouseButtonEffecting)
 		{
-			bIsSelecting = true;
 			SelectionStart = MouseCursorPosWorld;
 			SelectionEnd = SelectionStart;
 			SelectBox.Init();
@@ -497,8 +431,18 @@ FReply SWorldPartitionEditorGrid2D::OnMouseButtonUp(const FGeometry& MyGeometry,
 	{
 		FReply ReplyState = FReply::Handled();
 
-		if ((!bIsSelecting && !bIsDragging) && bIsRightMouseButtonEffecting)
-		{
+		const bool bHasMouseCapture = bIsDragSelecting || bIsPanning;
+		MouseCursorPos = MyGeometry.AbsoluteToLocal(MouseEvent.GetScreenSpacePosition());
+		MouseCursorPosWorld = ScreenToWorld.TransformPoint(MouseCursorPos);
+
+		if (!bHasMouseCapture && bIsRightMouseButtonEffecting)
+		{			
+			if (HoveredLoaderAdapter.IsValid() && !SelectedLoaderAdapters.Contains(HoveredLoaderAdapter))
+			{
+				SelectedLoaderAdapters.Reset();
+				SelectedLoaderAdapters.Add(HoveredLoaderAdapter);
+			}
+			
 			FMenuBuilder MenuBuilder(true, CommandList);
 
 			const FEditorCommands& Commands = FEditorCommands::Get();
@@ -508,7 +452,6 @@ FReply SWorldPartitionEditorGrid2D::OnMouseButtonUp(const FGeometry& MyGeometry,
 				MenuBuilder.AddMenuSeparator();
 				MenuBuilder.AddMenuEntry(Commands.LoadSelectedRegions);
 				MenuBuilder.AddMenuEntry(Commands.UnloadSelectedRegions);
-				MenuBuilder.AddMenuEntry(Commands.UnloadHoveredRegion);
 				MenuBuilder.AddMenuSeparator();
 				MenuBuilder.AddMenuEntry(Commands.ConvertSelectedRegionsToActors);
 			MenuBuilder.EndSection();
@@ -521,19 +464,47 @@ FReply SWorldPartitionEditorGrid2D::OnMouseButtonUp(const FGeometry& MyGeometry,
 			FSlateApplication::Get().PushMenu(AsShared(), WidgetPath, MenuBuilder.MakeWidget(), MouseEvent.GetScreenSpacePosition(), FPopupTransitionEffect(FPopupTransitionEffect::ContextMenu));
 		}
 
-		const bool HasMouseCapture = bIsSelecting || bIsDragging;
-
 		if (bIsLeftMouseButtonEffecting)
 		{
-			bIsSelecting = false;
+			FLoaderAdapterSet LoaderAdaptersToSelect;
+			if (bIsDragSelecting)
+			{
+				ForEachIntersectingLoaderAdapters(WorldPartition, SelectBox, [&LoaderAdaptersToSelect](UObject* AdapterObject)
+				{
+					LoaderAdaptersToSelect.Add(AdapterObject);
+					return true;
+				});
+			}
+			else if (HoveredLoaderAdapter.IsValid())
+			{
+				if (MouseEvent.IsControlDown() && SelectedLoaderAdapters.Contains(HoveredLoaderAdapter))
+				{
+					SelectedLoaderAdapters.Remove(HoveredLoaderAdapter);
+				}
+				else
+				{
+					LoaderAdaptersToSelect.Add(HoveredLoaderAdapter);
+				}
+			}
+
+			if (MouseEvent.IsControlDown())
+			{
+				SelectedLoaderAdapters.Append(LoaderAdaptersToSelect);
+			}
+			else
+			{
+				SelectedLoaderAdapters = MoveTemp(LoaderAdaptersToSelect);
+			}
+
+			bIsDragSelecting = false;
 		}
 
 		if (bIsRightMouseButtonEffecting)
 		{
-			bIsDragging = false;
+			bIsPanning = false;
 		}
 
-		if (HasMouseCapture && !bIsSelecting && !bIsDragging)
+		if (HasMouseCapture() && !bIsDragSelecting && !bIsPanning)
 		{
 			ReplyState.ReleaseMouseCapture();
 		}
@@ -562,28 +533,37 @@ FReply SWorldPartitionEditorGrid2D::OnMouseMove(const FGeometry& MyGeometry, con
 		const bool bIsRightMouseButtonDown = MouseEvent.IsMouseButtonDown(EKeys::RightMouseButton);
 		const bool bIsLeftMouseButtonDown = MouseEvent.IsMouseButtonDown(EKeys::LeftMouseButton);
 		const bool bIsMiddleMouseButtonDown = MouseEvent.IsMouseButtonDown(EKeys::MiddleMouseButton);
+		const bool bIsDragTrigger = CursorDelta.Size() > FSlateApplication::Get().GetDragTriggerDistance();
 
 		if (bIsLeftMouseButtonDown)
 		{
-			if (bIsSelecting)
+			if (!bIsDragSelecting && bIsDragTrigger)
+			{
+				bIsDragSelecting = true;
+			}
+
+			if (bIsDragSelecting)
 			{
 				SelectionEnd = MouseCursorPosWorld;
-				UpdateSelection();
+				UpdateSelectionBox();
+				return FReply::Handled();
 			}
 		}
 
-		if (bIsDragging || (bIsRightMouseButtonDown && (CursorDelta.Size() > 1.0f)))
+		if (bIsRightMouseButtonDown && !bIsDragSelecting)
 		{
-			if (!bIsDragging)
+			if (!bIsPanning && bIsDragTrigger)
 			{
+				bIsPanning = true;
 				LastMouseCursorPosWorldDrag = MouseCursorPosWorld;
-				bIsDragging = true;
 			}
 
-			Trans += (MouseCursorPosWorld - LastMouseCursorPosWorldDrag);
-
-			UpdateTransform();
-			return FReply::Handled();
+			if (bIsPanning)
+			{
+				Trans += (MouseCursorPosWorld - LastMouseCursorPosWorldDrag);
+				UpdateTransform();
+				return FReply::Handled();
+			}
 		}
 	}
 
@@ -604,7 +584,7 @@ FReply SWorldPartitionEditorGrid2D::OnMouseWheel(const FGeometry& MyGeometry, co
 
 FCursorReply SWorldPartitionEditorGrid2D::OnCursorQuery(const FGeometry& MyGeometry, const FPointerEvent& CursorEvent) const
 {
-	return FCursorReply::Cursor(bIsDragging ? EMouseCursor::None : EMouseCursor::Default);
+	return FCursorReply::Cursor(bIsPanning ? EMouseCursor::None : EMouseCursor::Default);
 }
 
 int32 SWorldPartitionEditorGrid2D::PaintGrid(const FGeometry& AllottedGeometry, const FSlateRect& MyCullingRect, FSlateWindowElementList& OutDrawElements, int32 LayerId) const
@@ -642,19 +622,39 @@ int32 SWorldPartitionEditorGrid2D::PaintGrid(const FGeometry& AllottedGeometry,
 	return LayerId + 1;
 }
 
-uint32 SWorldPartitionEditorGrid2D::PaintActors(const FGeometry& AllottedGeometry, const FSlateRect& MyCullingRect, FSlateWindowElementList& OutDrawElements, uint32 LayerId) const
+void SWorldPartitionEditorGrid2D::Tick(const FGeometry& AllottedGeometry, const double InCurrentTime, const float InDeltaTime)
 {
 	const FBox2D ViewRect(FVector2D(ForceInitToZero), AllottedGeometry.GetLocalSize());
 	const FBox2D WorldViewRect(ScreenToWorld.TransformPoint(ViewRect.Min), ScreenToWorld.TransformPoint(ViewRect.Max));
 	const FBox ViewRectWorld(FVector(WorldViewRect.Min.X, WorldViewRect.Min.Y, -HALF_WORLD_MAX), FVector(WorldViewRect.Max.X, WorldViewRect.Max.Y, HALF_WORLD_MAX));
 
-	TSet<FWorldPartitionActorDescViewBoundsProxy> ActorDescList;
+	ShownActorGuids.Reset();
+	ShownLoaderAdapters.Reset();
+	HighlightedLoaderAdapters.Reset();
+	HoveredLoaderAdapter = nullptr;
+	HoveredActor.Invalidate();
 	
-	TArray<IWorldPartitionActorLoaderInterface::ILoaderAdapter*> AllLoaderAdapters;
-	for (IWorldPartitionActorLoaderInterface::ILoaderAdapter* LoaderAdapter : WorldPartition->GetRegisteredEditorLoaderAdapters())
+	double MinLoaderDistance = DBL_MAX;
+	double MinActorDistance = DBL_MAX;
+	for (UWorldPartitionEditorLoaderAdapter* EditorLoaderAdapter : WorldPartition->GetRegisteredEditorLoaderAdapters())
 	{
+		check(EditorLoaderAdapter);
+		IWorldPartitionActorLoaderInterface::ILoaderAdapter* LoaderAdapter = EditorLoaderAdapter->GetLoaderAdapter();
 		check(LoaderAdapter);
-		AllLoaderAdapters.Add(LoaderAdapter);
+
+		ShownLoaderAdapters.Add(LoaderAdapter);
+				
+		double OutDistance = DBL_MAX;
+		if (LoaderAdapter->GetBoundingBox().IsSet() && IsBoundsHoveredXY(MouseCursorPosWorld, *LoaderAdapter->GetBoundingBox(), OutDistance) && OutDistance < MinLoaderDistance)
+		{
+			MinLoaderDistance = OutDistance;
+			HoveredLoaderAdapter = EditorLoaderAdapter;
+		}
+
+		if (SelectedLoaderAdapters.Contains(EditorLoaderAdapter))
+		{
+			HighlightedLoaderAdapters.Add(LoaderAdapter);
+		}
 	}
 
 	WorldPartition->EditorHash->ForEachIntersectingActor(ViewRectWorld, [&](FWorldPartitionActorDesc* ActorDesc)
@@ -665,17 +665,36 @@ uint32 SWorldPartitionEditorGrid2D::PaintActors(const FGeometry& AllottedGeometr
 		{
 			if (ActorDescViewProxy.GetIsSpatiallyLoaded())
 			{
-				ActorDescList.Emplace(ActorDescViewProxy);
+				ShownActorGuids.Add(ActorDesc->GetGuid());
+
+				double OutDistance = DBL_MAX;
+				if (IsBoundsHoveredXY(MouseCursorPosWorld, ActorDescViewProxy.GetBounds(), OutDistance) && OutDistance < MinActorDistance)
+				{
+					MinActorDistance = OutDistance;
+					HoveredActor = ActorDesc->GetGuid();
+				}
 			}
 		}
-		
+
 		if (ActorDesc->GetActorNativeClass()->ImplementsInterface(UWorldPartitionActorLoaderInterface::StaticClass()))
 		{
 			if (AActor* Actor = ActorDescViewProxy.GetActor())
 			{
 				if (IWorldPartitionActorLoaderInterface::ILoaderAdapter* LoaderAdapter = Cast<IWorldPartitionActorLoaderInterface>(Actor)->GetLoaderAdapter())
 				{
-					AllLoaderAdapters.Add(LoaderAdapter);
+					ShownLoaderAdapters.Add(LoaderAdapter);
+
+					double OutDistance = DBL_MAX;
+					if (LoaderAdapter->GetBoundingBox().IsSet() && IsBoundsHoveredXY(MouseCursorPosWorld, *LoaderAdapter->GetBoundingBox(), OutDistance) && OutDistance < MinLoaderDistance)
+					{
+						MinLoaderDistance = OutDistance;
+						HoveredLoaderAdapter = Cast<IWorldPartitionActorLoaderInterface>(Actor);
+					}
+
+					if (SelectedLoaderAdapters.Contains(Actor))
+					{
+						HighlightedLoaderAdapters.Add(LoaderAdapter);
+					}
 				}
 			}
 		}
@@ -686,10 +705,31 @@ uint32 SWorldPartitionEditorGrid2D::PaintActors(const FGeometry& AllottedGeometr
 	{
 		if (AActor* Actor = Cast<AActor>(*It))
 		{
-			if (FWorldPartitionActorDesc* ActorDesc = WorldPartition->GetActorDesc(Actor->GetActorGuid()))
-			{
-				ActorDescList.Emplace(ActorDesc);
-			}
+			ShownActorGuids.Add(Actor->GetActorGuid());
+		}
+	}
+}
+
+uint32 SWorldPartitionEditorGrid2D::PaintActors(const FGeometry& AllottedGeometry, const FSlateRect& MyCullingRect, FSlateWindowElementList& OutDrawElements, uint32 LayerId) const
+{
+	const FBox2D ViewRect(FVector2D(ForceInitToZero), AllottedGeometry.GetLocalSize());
+	const FBox2D WorldViewRect(ScreenToWorld.TransformPoint(ViewRect.Min), ScreenToWorld.TransformPoint(ViewRect.Max));
+	const FBox ViewRectWorld(FVector(WorldViewRect.Min.X, WorldViewRect.Min.Y, -HALF_WORLD_MAX), FVector(WorldViewRect.Max.X, WorldViewRect.Max.Y, HALF_WORLD_MAX));
+
+	TArray<FWorldPartitionActorDescViewBoundsProxy> ActorDescList;
+	ActorDescList.Reserve(ShownActorGuids.Num());
+
+	const IWorldPartitionActorLoaderInterface::ILoaderAdapter* LocalHoveredLoaderAdapter = nullptr;
+	if (IWorldPartitionActorLoaderInterface* HoveredLoaderAdapterInterface = Cast<IWorldPartitionActorLoaderInterface>(HoveredLoaderAdapter.Get()))
+	{
+		LocalHoveredLoaderAdapter = HoveredLoaderAdapterInterface->GetLoaderAdapter();
+	}
+		
+	for (const FGuid& ActorGuid : ShownActorGuids)
+	{
+		if (FWorldPartitionActorDesc* ActorDesc = WorldPartition->GetActorDesc(ActorGuid))
+		{
+			ActorDescList.Emplace(ActorDesc);
 		}
 	}
 
@@ -720,7 +760,7 @@ uint32 SWorldPartitionEditorGrid2D::PaintActors(const FGeometry& AllottedGeometr
 		}
 	};
 
-	if (AllLoaderAdapters.Num())
+	if (ShownLoaderAdapters.Num())
 	{
 		const FLinearColor LoadedActorColor(0.75f, 0.75f, 0.75f, 1.0f);
 		const FLinearColor UnloadedActorColor(0.5f, 0.5f, 0.5f, 1.0f);	
@@ -728,14 +768,14 @@ uint32 SWorldPartitionEditorGrid2D::PaintActors(const FGeometry& AllottedGeometr
 		TArray<FVector2D> LinePoints;
 		LinePoints.SetNum(5);
 
-		for (const IWorldPartitionActorLoaderInterface::ILoaderAdapter* LoaderAdapter : AllLoaderAdapters)
+		for (const IWorldPartitionActorLoaderInterface::ILoaderAdapter* LoaderAdapter : ShownLoaderAdapters)
 		{
 			if (LoaderAdapter->GetBoundingBox().IsSet())
 			{
-				const FBox ActorBounds = *LoaderAdapter->GetBoundingBox();
+				const FBox AdapterBounds = *LoaderAdapter->GetBoundingBox();
 
 				FVector Origin, Extent;
-				ActorBounds.GetCenterAndExtents(Origin, Extent);
+				AdapterBounds.GetCenterAndExtents(Origin, Extent);
 
 				const FVector2D TopLeftW = FVector2D(Origin - Extent);
 				const FVector2D BottomRightW = FVector2D(Origin + Extent);
@@ -781,7 +821,7 @@ uint32 SWorldPartitionEditorGrid2D::PaintActors(const FGeometry& AllottedGeometr
 						{
 							const FLinearColor LoaderColor = LoaderAdapter->GetColor().IsSet() ? *LoaderAdapter->GetColor() : FColor::White;
 							const FLinearColor OutlineColor(LoaderColor.R, LoaderColor.G, LoaderColor.B, LoaderColorGradient * FullScreenColorGradient);
-							const bool IsHighlighted = IsBoundsSelected(SelectBox, ActorBounds) || IsBoundsHovered(MouseCursorPos, ActorViewBox);
+							const bool IsHighlighted = HighlightedLoaderAdapters.Contains(LoaderAdapter) || LocalHoveredLoaderAdapter == LoaderAdapter;
 
 							LinePoints[0] = TopLeft;
 							LinePoints[1] = TopRight;
@@ -853,7 +893,7 @@ uint32 SWorldPartitionEditorGrid2D::PaintActors(const FGeometry& AllottedGeometr
 				const AActor* Actor = ActorDescView.GetActor();
 
 				const bool bIsSelected = Actor ? Actor->IsSelected() : false;
-				if (bIsSelected || IsBoundsHovered(MouseCursorPos, ActorViewBox))
+				if (bIsSelected || HoveredActor == ActorDescView.GetGuid())
 				{
 					ActorColor = FLinearColor::Yellow;
 
@@ -1089,7 +1129,7 @@ int32 SWorldPartitionEditorGrid2D::OnPaint(const FPaintArgs& Args, const FGeomet
 
 int32 SWorldPartitionEditorGrid2D::PaintSoftwareCursor(const FGeometry& AllottedGeometry, const FSlateRect& MyCullingRect, FSlateWindowElementList& OutDrawElements, int32 LayerId) const
 {
-	if (bIsDragging)
+	if (bIsPanning)
 	{
 		const FSlateBrush* Brush = FAppStyle::GetBrush(TEXT("SoftwareCursor_Grab"));
 
@@ -1154,7 +1194,7 @@ void SWorldPartitionEditorGrid2D::UpdateTransform() const
 	ScreenToWorld = WorldToScreen.Inverse();
 }
 
-void SWorldPartitionEditorGrid2D::UpdateSelection()
+void SWorldPartitionEditorGrid2D::UpdateSelectionBox()
 {
 	const FBox2D SelectBox2D(
 		FVector2D::Min(SelectionStart, SelectionEnd), 
@@ -1170,4 +1210,10 @@ void SWorldPartitionEditorGrid2D::UpdateSelection()
 	}
 }
 
-#undef LOCTEXT_NAMESPACE
\ No newline at end of file
+void SWorldPartitionEditorGrid2D::ClearSelection()
+{
+	SelectedLoaderAdapters.Empty();
+	SelectBox.Init();
+}
+
+#undef LOCTEXT_NAMESPACE
diff --git a/Engine/Source/Editor/WorldPartitionEditor/Public/WorldPartition/SWorldPartitionEditorGrid2D.h b/Engine/Source/Editor/WorldPartitionEditor/Public/WorldPartition/SWorldPartitionEditorGrid2D.h
--- a/Engine/Source/Editor/WorldPartitionEditor/Public/WorldPartition/SWorldPartitionEditorGrid2D.h
+++ b/Engine/Source/Editor/WorldPartitionEditor/Public/WorldPartition/SWorldPartitionEditorGrid2D.h
@@ -2,10 +2,13 @@
 #pragma once
 
 #include "CoreMinimal.h"
+#include "UObject/WeakInterfacePtr.h"
+#include "UObject/WeakObjectPtrTemplates.h"
 #include "Framework/Commands/Commands.h"
 #include "Framework/Commands/UICommandList.h"
 #include "Layout/ArrangedChildren.h"
 #include "SWorldPartitionEditorGrid.h"
+#include "WorldPartition/WorldPartitionActorLoaderInterface.h"
 
 class SWorldPartitionEditorGrid2D : public SWorldPartitionEditorGrid
 {
@@ -37,10 +40,10 @@ public:
 	virtual void CreateRegionFromSelection();
 	virtual void LoadSelectedRegions();
 	virtual void UnloadSelectedRegions();
-	virtual void UnloadHoveredRegion();
 	virtual void ConvertSelectedRegionsToActors();
 	void MoveCameraHere();
 
+	virtual void Tick(const FGeometry& AllottedGeometry, const double InCurrentTime, const float InDeltaTime) override;
 	virtual FReply OnMouseButtonDown(const FGeometry& MyGeometry, const FPointerEvent& MouseEvent) override;
 	virtual FReply OnMouseButtonUp(const FGeometry& MyGeometry, const FPointerEvent& MouseEvent) override;
 	virtual FReply OnMouseButtonDoubleClick(const FGeometry& InMyGeometry, const FPointerEvent& InMouseEvent) override;
@@ -57,10 +60,10 @@ public:
 	virtual int32 PaintSoftwareCursor(const FGeometry& AllottedGeometry, const FSlateRect& MyCullingRect, FSlateWindowElementList& OutDrawElements, int32 LayerId) const;
 
 	virtual FReply FocusSelection();
-
 protected:
 	void UpdateTransform() const;
-	void UpdateSelection();
+	void UpdateSelectionBox();
+	void ClearSelection();
 
 	const TSharedRef<FUICommandList> CommandList;
 
@@ -85,8 +88,8 @@ protected:
 	mutable FTransform2d WorldToScreen;
 	mutable FTransform2d ScreenToWorld;
 
-	bool bIsSelecting;
-	bool bIsDragging;
+	bool bIsDragSelecting;
+	bool bIsPanning;
 	bool bShowActors;
 	FVector2D MouseCursorPos;
 	FVector2D MouseCursorPosWorld;
@@ -95,4 +98,32 @@ protected:
 	FVector2D SelectionEnd;
 	FBox SelectBox;
 	FSlateFontInfo SmallLayoutFont;
+
+	struct FKeyFuncs : public BaseKeyFuncs<TWeakInterfacePtr<IWorldPartitionActorLoaderInterface>, TWeakInterfacePtr<IWorldPartitionActorLoaderInterface>, false>
+	{
+		static KeyInitType GetSetKey(ElementInitType Entry)
+		{
+			return Entry;
+		}
+
+		static bool Matches(KeyInitType A, KeyInitType B)
+		{
+			return A == B;
+		}
+
+		static uint32 GetKeyHash(KeyInitType Key)
+		{
+			return GetTypeHash(Key.GetWeakObjectPtr());
+		}
+	};
+
+	using FLoaderAdapterSet = TSet<TWeakInterfacePtr<IWorldPartitionActorLoaderInterface>, FKeyFuncs>;
+	FLoaderAdapterSet SelectedLoaderAdapters;
+	
+	// Updated every tick
+	TSet<FGuid> ShownActorGuids;
+	TSet<IWorldPartitionActorLoaderInterface::ILoaderAdapter*> ShownLoaderAdapters;
+	TSet<IWorldPartitionActorLoaderInterface::ILoaderAdapter*> HighlightedLoaderAdapters;
+	TWeakInterfacePtr<IWorldPartitionActorLoaderInterface> HoveredLoaderAdapter;
+	FGuid HoveredActor;
 };
\ No newline at end of file
diff --git a/Engine/Source/Runtime/Engine/Private/WorldPartition/WorldPartition.cpp b/Engine/Source/Runtime/Engine/Private/WorldPartition/WorldPartition.cpp
--- a/Engine/Source/Runtime/Engine/Private/WorldPartition/WorldPartition.cpp
+++ b/Engine/Source/Runtime/Engine/Private/WorldPartition/WorldPartition.cpp
@@ -597,9 +597,9 @@ void UWorldPartition::Uninitialize()
 
 		if (RegisteredEditorLoaderAdapters.Num())
 		{
-			for (IWorldPartitionActorLoaderInterface::ILoaderAdapter* RegisteredLoaderAdapter : RegisteredEditorLoaderAdapters)
+			for (UWorldPartitionEditorLoaderAdapter* RegisteredEditorLoaderAdapter : RegisteredEditorLoaderAdapters)
 			{
-				delete RegisteredLoaderAdapter;
+				RegisteredEditorLoaderAdapter->Release();
 			}
 
 			RegisteredEditorLoaderAdapters.Empty();
@@ -649,8 +649,8 @@ void UWorldPartition::OnPostBugItGoCalled(const FVector& Loc, const FRotator& Ro
 		const FVector LoadExtent(GLoadingRangeBugItGo, GLoadingRangeBugItGo, HALF_WORLD_MAX);
 		const FBox LoadCellsBox(Loc - LoadExtent, Loc + LoadExtent);
 
-		FLoaderAdapterShape* LoaderAdapter = CreateEditorLoaderAdapter<FLoaderAdapterShape>(World, LoadCellsBox, TEXT("BugItGo"));
-		LoaderAdapter->Load();
+		UWorldPartitionEditorLoaderAdapter* EditorLoaderAdapter = CreateEditorLoaderAdapter<FLoaderAdapterShape>(World, LoadCellsBox, TEXT("BugItGo"));
+		EditorLoaderAdapter->GetLoaderAdapter()->Load();
 	}
 #endif
 }
@@ -1232,8 +1232,10 @@ TArray<FBox> UWorldPartition::GetUserLoadedEditorRegions() const
 {
 	TArray<FBox> Result;
 
-	for (IWorldPartitionActorLoaderInterface::ILoaderAdapter* LoaderAdapter : RegisteredEditorLoaderAdapters)
+	for (UWorldPartitionEditorLoaderAdapter* EditorLoaderAdapter : RegisteredEditorLoaderAdapters)
 	{
+		IWorldPartitionActorLoaderInterface::ILoaderAdapter* LoaderAdapter = EditorLoaderAdapter->GetLoaderAdapter();
+		check(LoaderAdapter);
 		if (LoaderAdapter->IsLoaded() && LoaderAdapter->GetUserCreated())
 		{
 			Result.Add(*LoaderAdapter->GetBoundingBox());
@@ -1343,7 +1345,9 @@ void UWorldPartition::LoadLastLoadedRegions(const TArray<FBox>& EditorLastLoaded
 {
 	for (const FBox& EditorLastLoadedRegion : EditorLastLoadedRegions)
 	{
-		FLoaderAdapterShape* LoaderAdapter = CreateEditorLoaderAdapter<FLoaderAdapterShape>(World, EditorLastLoadedRegion, TEXT("Last Loaded Region"));
+		UWorldPartitionEditorLoaderAdapter* EditorLoaderAdapter = CreateEditorLoaderAdapter<FLoaderAdapterShape>(World, EditorLastLoadedRegion, TEXT("Last Loaded Region"));
+		IWorldPartitionActorLoaderInterface::ILoaderAdapter* LoaderAdapter = EditorLoaderAdapter->GetLoaderAdapter();
+		check(LoaderAdapter);
 		LoaderAdapter->SetUserCreated(true);
 		LoaderAdapter->Load();
 	}
diff --git a/Engine/Source/Runtime/Engine/Public/WorldPartition/WorldPartition.h b/Engine/Source/Runtime/Engine/Public/WorldPartition/WorldPartition.h
--- a/Engine/Source/Runtime/Engine/Public/WorldPartition/WorldPartition.h
+++ b/Engine/Source/Runtime/Engine/Public/WorldPartition/WorldPartition.h
@@ -16,6 +16,7 @@
 
 #if WITH_EDITOR
 #include "WorldPartition/WorldPartitionActorLoaderInterface.h"
+#include "WorldPartition/WorldPartitionEditorLoaderAdapter.h"
 #include "PackageSourceControlHelper.h"
 #include "CookPackageSplitter.h"
 #endif
@@ -296,26 +297,31 @@ private:
 public:
 	// Editor loader adapters management
 	template <typename T, typename... ArgsType>
-	T* CreateEditorLoaderAdapter(ArgsType&&... Args)
+	UWorldPartitionEditorLoaderAdapter* CreateEditorLoaderAdapter(ArgsType&&... Args)
 	{
-		T* LoaderAdapter = new T(Forward<ArgsType>(Args)...);
-		RegisteredEditorLoaderAdapters.Add(LoaderAdapter);
-		return LoaderAdapter;
+		UWorldPartitionEditorLoaderAdapter* EditorLoaderAdapter = NewObject<UWorldPartitionEditorLoaderAdapter>(GetTransientPackage());
+		EditorLoaderAdapter->SetLoaderAdapter(new T(Forward<ArgsType>(Args)...));
+		RegisteredEditorLoaderAdapters.Add(EditorLoaderAdapter);
+		return EditorLoaderAdapter;
 	}
 
-	void ReleaseEditorLoaderAdapter(IWorldPartitionActorLoaderInterface::ILoaderAdapter* LoaderAdapter)
+	void ReleaseEditorLoaderAdapter(UWorldPartitionEditorLoaderAdapter* EditorLoaderAdapter)
 	{
-		verify(RegisteredEditorLoaderAdapters.Remove(LoaderAdapter) != INDEX_NONE);
-		delete LoaderAdapter;
+		verify(RegisteredEditorLoaderAdapters.Remove(EditorLoaderAdapter) != INDEX_NONE);
+		EditorLoaderAdapter->Release();
 	}
 
-	const TSet<IWorldPartitionActorLoaderInterface::ILoaderAdapter*>& GetRegisteredEditorLoaderAdapters() const
+	const TSet<TObjectPtr<UWorldPartitionEditorLoaderAdapter>>& GetRegisteredEditorLoaderAdapters() const
 	{
 		return RegisteredEditorLoaderAdapters;
 	}
 
 private:
-	TSet<IWorldPartitionActorLoaderInterface::ILoaderAdapter*> RegisteredEditorLoaderAdapters;
+#endif
+
+#if WITH_EDITORONLY_DATA
+	UPROPERTY(transient, NonTransactional)
+	TSet<TObjectPtr<UWorldPartitionEditorLoaderAdapter>> RegisteredEditorLoaderAdapters;
 #endif
 
 #if !UE_BUILD_SHIPPING
diff --git a/Engine/Source/Runtime/Engine/Public/WorldPartition/WorldPartitionEditorLoaderAdapter.h b/Engine/Source/Runtime/Engine/Public/WorldPartition/WorldPartitionEditorLoaderAdapter.h
new file mode 100644
--- /dev/null
+++ b/Engine/Source/Runtime/Engine/Public/WorldPartition/WorldPartitionEditorLoaderAdapter.h
@@ -0,0 +1,43 @@
+// Copyright Epic Games, Inc. All Rights Reserved.
+#pragma once
+
+#include "CoreMinimal.h"
+#include "UObject/ObjectMacros.h"
+#include "UObject/Object.h"
+#include "WorldPartition/WorldPartitionActorLoaderInterface.h"
+#include "WorldPartitionEditorLoaderAdapter.generated.h"
+
+UCLASS()
+class ENGINE_API UWorldPartitionEditorLoaderAdapter : public UObject, public IWorldPartitionActorLoaderInterface
+{
+public:
+	GENERATED_BODY()
+
+#if WITH_EDITOR
+	virtual IWorldPartitionActorLoaderInterface::ILoaderAdapter* GetLoaderAdapter() override { return LoaderAdapter; }
+
+	virtual void BeginDestroy() override
+	{
+		Super::BeginDestroy();
+		check(!LoaderAdapter);
+	}
+private:
+	friend class UWorldPartition;
+	void SetLoaderAdapter(IWorldPartitionActorLoaderInterface::ILoaderAdapter* InLoaderAdapter)
+	{
+		check(!LoaderAdapter);
+		LoaderAdapter = InLoaderAdapter;
+	}
+
+	void Release()
+	{
+		if (LoaderAdapter)
+		{
+			delete LoaderAdapter;
+			LoaderAdapter = nullptr;
+		}
+	}
+
+	IWorldPartitionActorLoaderInterface::ILoaderAdapter* LoaderAdapter = nullptr;
+#endif
+};
\ No newline at end of file
