World Partition Editor:
- Cull editor loading adapters to view rectangle to reduce work.
- Don't treat loading adapters with no visible borders as selected.

#rnx

diff --git a/Engine/Source/Editor/WorldPartitionEditor/Private/WorldPartition/SWorldPartitionEditorGrid2D.cpp b/Engine/Source/Editor/WorldPartitionEditor/Private/WorldPartition/SWorldPartitionEditorGrid2D.cpp
--- a/Engine/Source/Editor/WorldPartitionEditor/Private/WorldPartition/SWorldPartitionEditorGrid2D.cpp
+++ b/Engine/Source/Editor/WorldPartitionEditor/Private/WorldPartition/SWorldPartitionEditorGrid2D.cpp
@@ -642,18 +642,24 @@ void SWorldPartitionEditorGrid2D::Tick(const FGeometry& AllottedGeometry, const
 		IWorldPartitionActorLoaderInterface::ILoaderAdapter* LoaderAdapter = EditorLoaderAdapter->GetLoaderAdapter();
 		check(LoaderAdapter);
 
-		ShownLoaderAdapters.Add(LoaderAdapter);
-				
-		double OutDistance = DBL_MAX;
-		if (LoaderAdapter->GetBoundingBox().IsSet() && IsBoundsHoveredXY(MouseCursorPosWorld, *LoaderAdapter->GetBoundingBox(), OutDistance) && OutDistance < MinLoaderDistance)
+		if (LoaderAdapter->GetBoundingBox().IsSet() && LoaderAdapter->GetBoundingBox()->IntersectXY(ViewRectWorld))
 		{
-			MinLoaderDistance = OutDistance;
-			HoveredLoaderAdapter = EditorLoaderAdapter;
-		}
+			ShownLoaderAdapters.Add(LoaderAdapter);
+				
+			if (!LoaderAdapter->GetBoundingBox()->IsInsideXY(ViewRectWorld))
+			{
+				double OutDistance = DBL_MAX;
+				if (IsBoundsHoveredXY(MouseCursorPosWorld, *LoaderAdapter->GetBoundingBox(), OutDistance) && OutDistance < MinLoaderDistance)
+				{
+					MinLoaderDistance = OutDistance;
+					HoveredLoaderAdapter = EditorLoaderAdapter;
+				}
 
-		if (SelectedLoaderAdapters.Contains(EditorLoaderAdapter))
-		{
-			HighlightedLoaderAdapters.Add(LoaderAdapter);
+				if (SelectedLoaderAdapters.Contains(EditorLoaderAdapter))
+				{
+					HighlightedLoaderAdapters.Add(LoaderAdapter);
+				}
+			}
 		}
 	}
 
