Fixed World Partition Editor creating seletion box was not working properly when moving mouse slowly.

diff --git a/Engine/Source/Editor/WorldPartitionEditor/Private/WorldPartition/SWorldPartitionEditorGrid2D.cpp b/Engine/Source/Editor/WorldPartitionEditor/Private/WorldPartition/SWorldPartitionEditorGrid2D.cpp
--- a/Engine/Source/Editor/WorldPartitionEditor/Private/WorldPartition/SWorldPartitionEditorGrid2D.cpp
+++ b/Engine/Source/Editor/WorldPartitionEditor/Private/WorldPartition/SWorldPartitionEditorGrid2D.cpp
@@ -156,6 +156,7 @@ void SWorldPartitionEditorGrid2D::Construct(const FArguments& InArgs)
 	// Defaults
 	Trans = FVector2D(0, 0);
 	Scale = 0.00133333332;
+	TotalMouseDelta = 0;
 
 	// UI
 	ChildSlot
@@ -400,6 +401,8 @@ FReply SWorldPartitionEditorGrid2D::OnMouseButtonDown(const FGeometry& MyGeometr
 	const bool bIsLeftMouseButtonDown = MouseEvent.IsMouseButtonDown(EKeys::LeftMouseButton);
 	const bool bIsMiddleMouseButtonDown = MouseEvent.IsMouseButtonDown(EKeys::MiddleMouseButton);
 
+	TotalMouseDelta = 0;
+
 	if (bIsLeftMouseButtonEffecting || bIsRightMouseButtonEffecting)
 	{
 		FReply ReplyState = FReply::Handled();
@@ -427,6 +430,8 @@ FReply SWorldPartitionEditorGrid2D::OnMouseButtonUp(const FGeometry& MyGeometry,
 	const bool bIsLeftMouseButtonDown = MouseEvent.IsMouseButtonDown(EKeys::LeftMouseButton);
 	const bool bIsMiddleMouseButtonDown = MouseEvent.IsMouseButtonDown(EKeys::MiddleMouseButton);
 
+	TotalMouseDelta = 0;
+
 	if (bIsLeftMouseButtonEffecting || bIsRightMouseButtonEffecting)
 	{
 		FReply ReplyState = FReply::Handled();
@@ -530,10 +535,12 @@ FReply SWorldPartitionEditorGrid2D::OnMouseMove(const FGeometry& MyGeometry, con
 
 	if (HasMouseCapture())
 	{
+		TotalMouseDelta += CursorDelta.Size();
+
 		const bool bIsRightMouseButtonDown = MouseEvent.IsMouseButtonDown(EKeys::RightMouseButton);
 		const bool bIsLeftMouseButtonDown = MouseEvent.IsMouseButtonDown(EKeys::LeftMouseButton);
 		const bool bIsMiddleMouseButtonDown = MouseEvent.IsMouseButtonDown(EKeys::MiddleMouseButton);
-		const bool bIsDragTrigger = CursorDelta.Size() > FSlateApplication::Get().GetDragTriggerDistance();
+		const bool bIsDragTrigger = TotalMouseDelta > FSlateApplication::Get().GetDragTriggerDistance();
 
 		if (bIsLeftMouseButtonDown)
 		{
diff --git a/Engine/Source/Editor/WorldPartitionEditor/Public/WorldPartition/SWorldPartitionEditorGrid2D.h b/Engine/Source/Editor/WorldPartitionEditor/Public/WorldPartition/SWorldPartitionEditorGrid2D.h
--- a/Engine/Source/Editor/WorldPartitionEditor/Public/WorldPartition/SWorldPartitionEditorGrid2D.h
+++ b/Engine/Source/Editor/WorldPartitionEditor/Public/WorldPartition/SWorldPartitionEditorGrid2D.h
@@ -98,6 +98,7 @@ protected:
 	FVector2D SelectionEnd;
 	FBox SelectBox;
 	FSlateFontInfo SmallLayoutFont;
+	float TotalMouseDelta;
 
 	struct FKeyFuncs : public BaseKeyFuncs<TWeakInterfacePtr<IWorldPartitionActorLoaderInterface>, TWeakInterfacePtr<IWorldPartitionActorLoaderInterface>, false>
 	{
