World Partition Editor:
- Show transient actor loaders in the editor window, mainly spawned by blutilities or python helper scripts.
- Refactored hover selection to stack hovered regions and run through the stack when exiting a region.

#rnx

diff --git a/Engine/Source/Editor/WorldPartitionEditor/Private/WorldPartition/SWorldPartitionEditorGrid2D.cpp b/Engine/Source/Editor/WorldPartitionEditor/Private/WorldPartition/SWorldPartitionEditorGrid2D.cpp
--- a/Engine/Source/Editor/WorldPartitionEditor/Private/WorldPartition/SWorldPartitionEditorGrid2D.cpp
+++ b/Engine/Source/Editor/WorldPartitionEditor/Private/WorldPartition/SWorldPartitionEditorGrid2D.cpp
@@ -3,6 +3,7 @@
 #include "WorldPartition/SWorldPartitionEditorGrid2D.h"
 #include "Brushes/SlateColorBrush.h"
 #include "Engine/Selection.h"
+#include "EngineUtils.h"
 #include "GameFramework/WorldSettings.h"
 #include "Modules/ModuleManager.h"
 #include "WorldPartition/WorldPartition.h"
@@ -33,18 +34,9 @@ static bool IsBoundsSelected(const FBox& SelectBox, const FBox& Bounds)
 	return SelectBox.IsValid && Bounds.IntersectXY(SelectBox) && !Bounds.IsInsideXY(SelectBox);
 }
 
-static bool IsBoundsHoveredXY(FVector2D Point, const FBox& Bounds, double& OutDistanceToEdge)
+static bool IsBoundsHovered(FVector2D Point, const FBox& Bounds)
 {
-	const FBox2D Bounds2D(FVector2D(Bounds.Min.X, Bounds.Min.Y), FVector2D(Bounds.Max.X, Bounds.Max.Y));
-	if (Bounds2D.IsInsideOrOn(Point))
-	{
-		const double DistX = FMath::Min(Point.X - Bounds.Min.X, Bounds.Max.X - Point.X);
-		const double DistY = FMath::Min(Point.Y - Bounds.Min.Y, Bounds.Max.Y - Point.Y);
-		OutDistanceToEdge = FMath::Min(DistX, DistY);
-		return true;
-	}
-
-	return false;
+	return Bounds.IsInsideOrOnXY(FVector(Point.X, Point.Y, 0));
 }
 
 template <class T>
@@ -253,7 +245,7 @@ void SWorldPartitionEditorGrid2D::Construct(const FArguments& InArgs)
 
 	auto CanLoadUnloadSelectedRegions = [this](bool bLoad)
 	{
-		for (const TWeakInterfacePtr<IWorldPartitionActorLoaderInterface>& SelectedLoaderAdapter : SelectedLoaderAdapters)
+		for (const TWeakInterfacePtr<IWorldPartitionActorLoaderInterface>& SelectedLoaderAdapter : SelectedLoaderInterfaces)
 		{
 			if (IWorldPartitionActorLoaderInterface* LoaderInterface = SelectedLoaderAdapter.Get())
 			{
@@ -269,7 +261,7 @@ void SWorldPartitionEditorGrid2D::Construct(const FArguments& InArgs)
 
 	auto CanConvertSelectedRegionsToActors = [this]()
 	{
-		for (const TWeakInterfacePtr<IWorldPartitionActorLoaderInterface>& SelectedLoaderAdapter : SelectedLoaderAdapters)
+		for (const TWeakInterfacePtr<IWorldPartitionActorLoaderInterface>& SelectedLoaderAdapter : SelectedLoaderInterfaces)
 		{
 			if (UWorldPartitionEditorLoaderAdapter* EditorLoaderAdapter = Cast<UWorldPartitionEditorLoaderAdapter>(SelectedLoaderAdapter.Get()))
 			{
@@ -297,7 +289,7 @@ void SWorldPartitionEditorGrid2D::CreateRegionFromSelection()
 	EditorLoaderAdapter->GetLoaderAdapter()->Load();
 	
 	ClearSelection();
-	SelectedLoaderAdapters.Add(EditorLoaderAdapter);
+	SelectedLoaderInterfaces.Add(EditorLoaderAdapter);
 
 	GEditor->RedrawLevelEditingViewports();
 	Refresh();
@@ -305,7 +297,7 @@ void SWorldPartitionEditorGrid2D::CreateRegionFromSelection()
 
 void SWorldPartitionEditorGrid2D::LoadSelectedRegions()
 {
-	for (const TWeakInterfacePtr<IWorldPartitionActorLoaderInterface>& SelectedLoaderAdapter : SelectedLoaderAdapters)
+	for (const TWeakInterfacePtr<IWorldPartitionActorLoaderInterface>& SelectedLoaderAdapter : SelectedLoaderInterfaces)
 	{
 		if (IWorldPartitionActorLoaderInterface* LoaderInterface = SelectedLoaderAdapter.Get())
 		{
@@ -321,8 +313,8 @@ void SWorldPartitionEditorGrid2D::LoadSelectedRegions()
 
 void SWorldPartitionEditorGrid2D::UnloadSelectedRegions()
 {
-	FLoaderAdapterSet CopySelectedLoaderAdapters(SelectedLoaderAdapters);
-	for (const TWeakInterfacePtr<IWorldPartitionActorLoaderInterface>& SelectedLoaderAdapter : CopySelectedLoaderAdapters)
+	FLoaderInterfaceSet CopySelectedLoaderInterfaces(SelectedLoaderInterfaces);
+	for (const TWeakInterfacePtr<IWorldPartitionActorLoaderInterface>& SelectedLoaderAdapter : CopySelectedLoaderInterfaces)
 	{
 		if (IWorldPartitionActorLoaderInterface* LoaderInterface = SelectedLoaderAdapter.Get())
 		{
@@ -330,7 +322,7 @@ void SWorldPartitionEditorGrid2D::UnloadSelectedRegions()
 
 			if (UWorldPartitionEditorLoaderAdapter* EditorLoaderAdapter = Cast<UWorldPartitionEditorLoaderAdapter>(LoaderInterface))
 			{
-				SelectedLoaderAdapters.Remove(EditorLoaderAdapter);
+				SelectedLoaderInterfaces.Remove(EditorLoaderAdapter);
 				WorldPartition->ReleaseEditorLoaderAdapter(EditorLoaderAdapter);
 			}
 		}
@@ -344,10 +336,10 @@ void SWorldPartitionEditorGrid2D::UnloadSelectedRegions()
 
 void SWorldPartitionEditorGrid2D::ConvertSelectedRegionsToActors()
 {
-	FLoaderAdapterSet TmpSelectedLoaderAdapters = MoveTemp(SelectedLoaderAdapters);
+	FLoaderInterfaceSet TmpSelectedLoaderInterfaces = MoveTemp(SelectedLoaderInterfaces);
 	ClearSelection();
 
-	for (const TWeakInterfacePtr<IWorldPartitionActorLoaderInterface>& SelectedLoaderAdapter : TmpSelectedLoaderAdapters)
+	for (const TWeakInterfacePtr<IWorldPartitionActorLoaderInterface>& SelectedLoaderAdapter : TmpSelectedLoaderInterfaces)
 	{
 		if (UWorldPartitionEditorLoaderAdapter* EditorLoaderAdapter = Cast<UWorldPartitionEditorLoaderAdapter>(SelectedLoaderAdapter.Get()))
 		{
@@ -373,7 +365,7 @@ void SWorldPartitionEditorGrid2D::ConvertSelectedRegionsToActors()
 		}
 		else
 		{
-			SelectedLoaderAdapters.Add(SelectedLoaderAdapter);
+			SelectedLoaderInterfaces.Add(SelectedLoaderAdapter);
 		}
 	}
 
@@ -442,10 +434,10 @@ FReply SWorldPartitionEditorGrid2D::OnMouseButtonUp(const FGeometry& MyGeometry,
 
 		if (!bHasMouseCapture && bIsRightMouseButtonEffecting)
 		{			
-			if (HoveredLoaderAdapter.IsValid() && !SelectedLoaderAdapters.Contains(HoveredLoaderAdapter))
+			if (HoveredLoaderInterface.IsValid() && !SelectedLoaderInterfaces.Contains(HoveredLoaderInterface))
 			{
-				SelectedLoaderAdapters.Reset();
-				SelectedLoaderAdapters.Add(HoveredLoaderAdapter);
+				SelectedLoaderInterfaces.Reset();
+				SelectedLoaderInterfaces.Add(HoveredLoaderInterface);
 			}
 			
 			FMenuBuilder MenuBuilder(true, CommandList);
@@ -471,7 +463,7 @@ FReply SWorldPartitionEditorGrid2D::OnMouseButtonUp(const FGeometry& MyGeometry,
 
 		if (bIsLeftMouseButtonEffecting)
 		{
-			FLoaderAdapterSet LoaderAdaptersToSelect;
+			FLoaderInterfaceSet LoaderAdaptersToSelect;
 			if (bIsDragSelecting)
 			{
 				ForEachIntersectingLoaderAdapters(WorldPartition, SelectBox, [&LoaderAdaptersToSelect](UObject* AdapterObject)
@@ -480,25 +472,25 @@ FReply SWorldPartitionEditorGrid2D::OnMouseButtonUp(const FGeometry& MyGeometry,
 					return true;
 				});
 			}
-			else if (HoveredLoaderAdapter.IsValid())
+			else if (HoveredLoaderInterface.IsValid())
 			{
-				if (MouseEvent.IsControlDown() && SelectedLoaderAdapters.Contains(HoveredLoaderAdapter))
+				if (MouseEvent.IsControlDown() && SelectedLoaderInterfaces.Contains(HoveredLoaderInterface))
 				{
-					SelectedLoaderAdapters.Remove(HoveredLoaderAdapter);
+					SelectedLoaderInterfaces.Remove(HoveredLoaderInterface);
 				}
 				else
 				{
-					LoaderAdaptersToSelect.Add(HoveredLoaderAdapter);
+					LoaderAdaptersToSelect.Add(HoveredLoaderInterface);
 				}
 			}
 
 			if (MouseEvent.IsControlDown())
 			{
-				SelectedLoaderAdapters.Append(LoaderAdaptersToSelect);
+				SelectedLoaderInterfaces.Append(LoaderAdaptersToSelect);
 			}
 			else
 			{
-				SelectedLoaderAdapters = MoveTemp(LoaderAdaptersToSelect);
+				SelectedLoaderInterfaces = MoveTemp(LoaderAdaptersToSelect);
 			}
 
 			bIsDragSelecting = false;
@@ -636,10 +628,8 @@ void SWorldPartitionEditorGrid2D::Tick(const FGeometry& AllottedGeometry, const
 	const FBox ViewRectWorld(FVector(WorldViewRect.Min.X, WorldViewRect.Min.Y, -HALF_WORLD_MAX), FVector(WorldViewRect.Max.X, WorldViewRect.Max.Y, HALF_WORLD_MAX));
 
 	ShownActorGuids.Reset();
-	ShownLoaderAdapters.Reset();
-	HighlightedLoaderAdapters.Reset();
-	HoveredLoaderAdapter = nullptr;
-	HoveredActor.Invalidate();
+	ShownLoaderInterfaces.Reset();
+	HighlightedLoaderInterfaces.Reset();
 	
 	double MinLoaderDistance = DBL_MAX;
 	double MinActorDistance = DBL_MAX;
@@ -651,67 +641,97 @@ void SWorldPartitionEditorGrid2D::Tick(const FGeometry& AllottedGeometry, const
 
 		if (LoaderAdapter->GetBoundingBox().IsSet() && LoaderAdapter->GetBoundingBox()->IntersectXY(ViewRectWorld))
 		{
-			ShownLoaderAdapters.Add(LoaderAdapter);
-				
-			if (!LoaderAdapter->GetBoundingBox()->IsInsideXY(ViewRectWorld))
+			ShownLoaderInterfaces.Add(EditorLoaderAdapter);
+		}
+	}
+
+	WorldPartition->EditorHash->ForEachIntersectingActor(ViewRectWorld, [&](FWorldPartitionActorDesc* ActorDesc)
+	{
+		FWorldPartitionActorDescViewBoundsProxy ActorDescViewProxy(ActorDesc);
+
+		if (ActorDesc->GetActorNativeClass()->ImplementsInterface(UWorldPartitionActorLoaderInterface::StaticClass()))
+		{
+			if (AActor* Actor = ActorDescViewProxy.GetActor())
 			{
-				double OutDistance = DBL_MAX;
-				if (IsBoundsHoveredXY(MouseCursorPosWorld, *LoaderAdapter->GetBoundingBox(), OutDistance) && OutDistance < MinLoaderDistance)
+				if (IWorldPartitionActorLoaderInterface::ILoaderAdapter* LoaderAdapter = Cast<IWorldPartitionActorLoaderInterface>(Actor)->GetLoaderAdapter())
 				{
-					MinLoaderDistance = OutDistance;
-					HoveredLoaderAdapter = EditorLoaderAdapter;
+					ShownLoaderInterfaces.Add(Actor);
 				}
+			}
+		}
+	});
 
-				if (SelectedLoaderAdapters.Contains(EditorLoaderAdapter))
-				{
-					HighlightedLoaderAdapters.Add(LoaderAdapter);
-				}
+	// Also include transient actor loader adapters that might have been spawned by blutilities, etc. Since these actors can't be saved because they are transient,
+	// they will never get an actor descriptor so they will never appear in the world partition editor.
+	for (TActorIterator<AActor> ActorIt(World); ActorIt; ++ActorIt)
+	{
+		if (AActor* Actor = *ActorIt; Actor && Actor->HasAnyFlags(RF_Transient) && Actor->Implements<UWorldPartitionActorLoaderInterface>())
+		{
+			if (IWorldPartitionActorLoaderInterface::ILoaderAdapter* LoaderAdapter = Cast<IWorldPartitionActorLoaderInterface>(Actor)->GetLoaderAdapter())
+			{
+				ShownLoaderInterfaces.Add(Actor);
 			}
 		}
 	}
 
-	WorldPartition->EditorHash->ForEachIntersectingActor(ViewRectWorld, [&](FWorldPartitionActorDesc* ActorDesc)
+	FLoaderInterfaceSet LastHoveredLoaderInterfaces = MoveTemp(HoveredLoaderInterfaces);
+	for (const FLoaderInterface& LoaderInterface : ShownLoaderInterfaces)
 	{
-		FWorldPartitionActorDescViewBoundsProxy ActorDescViewProxy(ActorDesc);
+		const IWorldPartitionActorLoaderInterface::ILoaderAdapter* LoaderAdapter = LoaderInterface->GetLoaderAdapter();
 
-		if (bShowActors)
+		if (!LoaderAdapter->GetBoundingBox()->IsInsideXY(ViewRectWorld))
 		{
-			if (ActorDescViewProxy.GetIsSpatiallyLoaded())
+			if (IsBoundsHovered(MouseCursorPosWorld, *LoaderAdapter->GetBoundingBox()))
 			{
-				ShownActorGuids.Add(ActorDesc->GetGuid());
+				HoveredLoaderInterfaces.Add(LoaderInterface);
+			}
 
-				double OutDistance = DBL_MAX;
-				if (IsBoundsHoveredXY(MouseCursorPosWorld, ActorDescViewProxy.GetBounds(), OutDistance) && OutDistance < MinActorDistance)
-				{
-					MinActorDistance = OutDistance;
-					HoveredActor = ActorDesc->GetGuid();
-				}
+			if (SelectedLoaderInterfaces.Contains(LoaderInterface))
+			{
+				HighlightedLoaderInterfaces.Add(LoaderInterface);
 			}
 		}
+	}
 
-		if (ActorDesc->GetActorNativeClass()->ImplementsInterface(UWorldPartitionActorLoaderInterface::StaticClass()))
+	FLoaderInterfaceSet EnteredHoveredLoaderInterfaces = HoveredLoaderInterfaces.Difference(LastHoveredLoaderInterfaces);
+	FLoaderInterfaceSet ExitedHoveredLoaderInterfaces = LastHoveredLoaderInterfaces.Difference(HoveredLoaderInterfaces);
+
+	if (EnteredHoveredLoaderInterfaces.Num())
+	{
+		if (HoveredLoaderInterface.IsValid())
 		{
-			if (AActor* Actor = ActorDescViewProxy.GetActor())
-			{
-				if (IWorldPartitionActorLoaderInterface::ILoaderAdapter* LoaderAdapter = Cast<IWorldPartitionActorLoaderInterface>(Actor)->GetLoaderAdapter())
-				{
-					ShownLoaderAdapters.Add(LoaderAdapter);
+			HoveredLoaderInterfacesStack.Push(HoveredLoaderInterface);
+		}
 
-					double OutDistance = DBL_MAX;
-					if (LoaderAdapter->GetBoundingBox().IsSet() && IsBoundsHoveredXY(MouseCursorPosWorld, *LoaderAdapter->GetBoundingBox(), OutDistance) && OutDistance < MinLoaderDistance)
-					{
-						MinLoaderDistance = OutDistance;
-						HoveredLoaderAdapter = Cast<IWorldPartitionActorLoaderInterface>(Actor);
-					}
+		HoveredLoaderInterface = *FLoaderInterfaceSet::TIterator(EnteredHoveredLoaderInterfaces);
+	}
+	else if (ExitedHoveredLoaderInterfaces.Num())
+	{
+		if (ExitedHoveredLoaderInterfaces.Contains(HoveredLoaderInterface))
+		{
+			HoveredLoaderInterfaces.Remove(HoveredLoaderInterface);
+			HoveredLoaderInterface = nullptr;
 
-					if (SelectedLoaderAdapters.Contains(Actor))
+			// Go back in the hovered stack if possible
+			while (!HoveredLoaderInterfacesStack.IsEmpty())
+			{
+				if (FLoaderInterface StackedHoveredLoaderInterface = HoveredLoaderInterfacesStack.Pop(); StackedHoveredLoaderInterface.IsValid())
+				{
+					if (IsBoundsHovered(MouseCursorPosWorld, *StackedHoveredLoaderInterface->GetLoaderAdapter()->GetBoundingBox()))
 					{
-						HighlightedLoaderAdapters.Add(LoaderAdapter);
+						HoveredLoaderInterface = StackedHoveredLoaderInterface;
+						break;
 					}
 				}
 			}
+
+			// Last resort, take the first one in the list
+			if (!HoveredLoaderInterface.IsValid() && HoveredLoaderInterfaces.Num())
+			{
+				HoveredLoaderInterface = *FLoaderInterfaceSet::TIterator(HoveredLoaderInterfaces);
+			}
 		}
-	});
+	}
 
 	// Include selected actors
 	for (FSelectionIterator It = GEditor->GetSelectedActorIterator(); It; ++It)
@@ -733,9 +753,9 @@ uint32 SWorldPartitionEditorGrid2D::PaintActors(const FGeometry& AllottedGeometr
 	ActorDescList.Reserve(ShownActorGuids.Num());
 
 	const IWorldPartitionActorLoaderInterface::ILoaderAdapter* LocalHoveredLoaderAdapter = nullptr;
-	if (IWorldPartitionActorLoaderInterface* HoveredLoaderAdapterInterface = Cast<IWorldPartitionActorLoaderInterface>(HoveredLoaderAdapter.Get()))
+	if (IWorldPartitionActorLoaderInterface* HoveredLoaderAdapterInterface = Cast<IWorldPartitionActorLoaderInterface>(HoveredLoaderInterface.Get()))
 	{
-		LocalHoveredLoaderAdapter = HoveredLoaderAdapterInterface->GetLoaderAdapter();
+		LocalHoveredLoaderAdapter = HoveredLoaderInterface->GetLoaderAdapter();
 	}
 		
 	for (const FGuid& ActorGuid : ShownActorGuids)
@@ -773,7 +793,7 @@ uint32 SWorldPartitionEditorGrid2D::PaintActors(const FGeometry& AllottedGeometr
 		}
 	};
 
-	if (ShownLoaderAdapters.Num())
+	if (ShownLoaderInterfaces.Num())
 	{
 		const FLinearColor LoadedActorColor(0.75f, 0.75f, 0.75f, 1.0f);
 		const FLinearColor UnloadedActorColor(0.5f, 0.5f, 0.5f, 1.0f);	
@@ -781,8 +801,10 @@ uint32 SWorldPartitionEditorGrid2D::PaintActors(const FGeometry& AllottedGeometr
 		TArray<FVector2D> LinePoints;
 		LinePoints.SetNum(5);
 
-		for (const IWorldPartitionActorLoaderInterface::ILoaderAdapter* LoaderAdapter : ShownLoaderAdapters)
+		for (const FLoaderInterface& LoaderInterface : ShownLoaderInterfaces)
 		{
+			const IWorldPartitionActorLoaderInterface::ILoaderAdapter* LoaderAdapter = LoaderInterface->GetLoaderAdapter();
+
 			if (LoaderAdapter->GetBoundingBox().IsSet())
 			{
 				const FBox AdapterBounds = *LoaderAdapter->GetBoundingBox();
@@ -812,13 +834,14 @@ uint32 SWorldPartitionEditorGrid2D::PaintActors(const FGeometry& AllottedGeometr
 					{
 						const FPaintGeometry ActorGeometry = AllottedGeometry.ToPaintGeometry(TopLeft, BottomRight - TopLeft);
 						const float LoaderColorGradient = FMath::Min((ActorViewBox.GetArea() - MinimumAreaCull) / AreaFadeDistance, 1.0f);
+						const FLinearColor LoaderColor = LoaderAdapter->GetColor().IsSet() ? *LoaderAdapter->GetColor() : FColor::White;
 
 						// Highlight
 						{
 							const FSlateColorBrush LoadedBrush(FLinearColor::White);
-							const FSlateColorBrush UnloadedBrush(FLinearColor::Black);
-							const FLinearColor LoadedColor(1.0f, 1.0f, 1.0f, 0.25f * LoaderColorGradient * FullScreenColorGradient);
-							const FLinearColor UnloadedColor(0.0f, 0.0f, 0.0f, 0.25f * LoaderColorGradient * FullScreenColorGradient);
+							const FSlateColorBrush UnloadedBrush(FLinearColor::Gray);
+							const FLinearColor LoadedColor(LoaderColor.R, LoaderColor.G, LoaderColor.B, 0.25f * LoaderColorGradient * FullScreenColorGradient);
+							const FLinearColor UnloadedColor(LoaderColor.R * 0.15f, LoaderColor.G * 0.15f, LoaderColor.B * 0.15f, 0.25f * LoaderColorGradient * FullScreenColorGradient);
 
 							FSlateDrawElement::MakeBox(
 								OutDrawElements,
@@ -832,9 +855,8 @@ uint32 SWorldPartitionEditorGrid2D::PaintActors(const FGeometry& AllottedGeometr
 
 						// Outline
 						{
-							const FLinearColor LoaderColor = LoaderAdapter->GetColor().IsSet() ? *LoaderAdapter->GetColor() : FColor::White;
-							const FLinearColor OutlineColor(LoaderColor.R, LoaderColor.G, LoaderColor.B, LoaderColorGradient * FullScreenColorGradient);
-							const bool IsHighlighted = HighlightedLoaderAdapters.Contains(LoaderAdapter) || LocalHoveredLoaderAdapter == LoaderAdapter;
+							const bool IsHighlighted = HighlightedLoaderInterfaces.Contains(LoaderInterface) || LocalHoveredLoaderAdapter == LoaderAdapter;
+							const FLinearColor OutlineColor = IsHighlighted ? FLinearColor::Yellow : FLinearColor::White;
 
 							LinePoints[0] = TopLeft;
 							LinePoints[1] = TopRight;
@@ -906,7 +928,7 @@ uint32 SWorldPartitionEditorGrid2D::PaintActors(const FGeometry& AllottedGeometr
 				const AActor* Actor = ActorDescView.GetActor();
 
 				const bool bIsSelected = Actor ? Actor->IsSelected() : false;
-				if (bIsSelected || HoveredActor == ActorDescView.GetGuid())
+				if (bIsSelected)
 				{
 					ActorColor = FLinearColor::Yellow;
 
@@ -1225,7 +1247,7 @@ void SWorldPartitionEditorGrid2D::UpdateSelectionBox()
 
 void SWorldPartitionEditorGrid2D::ClearSelection()
 {
-	SelectedLoaderAdapters.Empty();
+	SelectedLoaderInterfaces.Empty();
 	SelectBox.Init();
 }
 
diff --git a/Engine/Source/Editor/WorldPartitionEditor/Public/WorldPartition/SWorldPartitionEditorGrid2D.h b/Engine/Source/Editor/WorldPartitionEditor/Public/WorldPartition/SWorldPartitionEditorGrid2D.h
--- a/Engine/Source/Editor/WorldPartitionEditor/Public/WorldPartition/SWorldPartitionEditorGrid2D.h
+++ b/Engine/Source/Editor/WorldPartitionEditor/Public/WorldPartition/SWorldPartitionEditorGrid2D.h
@@ -118,13 +118,17 @@ protected:
 		}
 	};
 
-	using FLoaderAdapterSet = TSet<TWeakInterfacePtr<IWorldPartitionActorLoaderInterface>, FKeyFuncs>;
-	FLoaderAdapterSet SelectedLoaderAdapters;
+	using FLoaderInterface = TWeakInterfacePtr<IWorldPartitionActorLoaderInterface>;
+	using FLoaderInterfaceSet = TSet<FLoaderInterface, FKeyFuncs>;
+	using FLoaderInterfaceStack = TArray<FLoaderInterface>;
+	
+	FLoaderInterfaceSet SelectedLoaderInterfaces;
 	
 	// Updated every tick
 	TSet<FGuid> ShownActorGuids;
-	TSet<IWorldPartitionActorLoaderInterface::ILoaderAdapter*> ShownLoaderAdapters;
-	TSet<IWorldPartitionActorLoaderInterface::ILoaderAdapter*> HighlightedLoaderAdapters;
-	TWeakInterfacePtr<IWorldPartitionActorLoaderInterface> HoveredLoaderAdapter;
-	FGuid HoveredActor;
+	FLoaderInterfaceSet ShownLoaderInterfaces;
+	FLoaderInterfaceSet HighlightedLoaderInterfaces;
+	FLoaderInterfaceSet HoveredLoaderInterfaces;
+	FLoaderInterfaceStack HoveredLoaderInterfacesStack;
+	FLoaderInterface HoveredLoaderInterface;
 };
\ No newline at end of file
