World Partition: don't use actor descriptor proxy when not needed + fix code that was removed accidentally during previous merge.

diff --git a/Engine/Source/Editor/WorldPartitionEditor/Private/WorldPartition/SWorldPartitionEditorGrid2D.cpp b/Engine/Source/Editor/WorldPartitionEditor/Private/WorldPartition/SWorldPartitionEditorGrid2D.cpp
--- a/Engine/Source/Editor/WorldPartitionEditor/Private/WorldPartition/SWorldPartitionEditorGrid2D.cpp
+++ b/Engine/Source/Editor/WorldPartitionEditor/Private/WorldPartition/SWorldPartitionEditorGrid2D.cpp
@@ -647,11 +647,17 @@ void SWorldPartitionEditorGrid2D::Tick(const FGeometry& AllottedGeometry, const
 
 	WorldPartition->EditorHash->ForEachIntersectingActor(ViewRectWorld, [&](FWorldPartitionActorDesc* ActorDesc)
 	{
-		FWorldPartitionActorDescViewBoundsProxy ActorDescViewProxy(ActorDesc);
+		if (bShowActors)
+		{
+			if (ActorDesc->GetIsSpatiallyLoaded())
+			{
+				ShownActorGuids.Add(ActorDesc->GetGuid());
+			}
+		}
 
 		if (ActorDesc->GetActorNativeClass()->ImplementsInterface(UWorldPartitionActorLoaderInterface::StaticClass()))
 		{
-			if (AActor* Actor = ActorDescViewProxy.GetActor())
+			if (AActor* Actor = ActorDesc->GetActor())
 			{
 				if (IWorldPartitionActorLoaderInterface::ILoaderAdapter* LoaderAdapter = Cast<IWorldPartitionActorLoaderInterface>(Actor)->GetLoaderAdapter())
 				{
