World Partition Editor: change UI from "Create a loading region from the selection" to "Load region from selection".

diff --git a/Engine/Source/Editor/WorldPartitionEditor/Private/WorldPartition/SWorldPartitionEditorGrid2D.cpp b/Engine/Source/Editor/WorldPartitionEditor/Private/WorldPartition/SWorldPartitionEditorGrid2D.cpp
--- a/Engine/Source/Editor/WorldPartitionEditor/Private/WorldPartition/SWorldPartitionEditorGrid2D.cpp
+++ b/Engine/Source/Editor/WorldPartitionEditor/Private/WorldPartition/SWorldPartitionEditorGrid2D.cpp
@@ -114,7 +114,7 @@ SWorldPartitionEditorGrid2D::FEditorCommands::FEditorCommands()
 
 void SWorldPartitionEditorGrid2D::FEditorCommands::RegisterCommands()
 {
-	UI_COMMAND(CreateRegionFromSelection, "Create Loading Region From Selection", "Create a loading region from the selection.", EUserInterfaceActionType::Button, FInputChord());
+	UI_COMMAND(CreateRegionFromSelection, "Load Region From Selection", "Load region from selection.", EUserInterfaceActionType::Button, FInputChord());
 	UI_COMMAND(LoadSelectedRegions, "Load Selected Regions", "Load the selected regions.", EUserInterfaceActionType::Button, FInputChord());
 	UI_COMMAND(UnloadSelectedRegions, "Unload Selected Regions", "Unload the selected regions.", EUserInterfaceActionType::Button, FInputChord());
 	UI_COMMAND(ConvertSelectedRegionsToActors, "Convert Selected Regions To Actors", "Convert the selected regions to actors.", EUserInterfaceActionType::Button, FInputChord());
