World Partition: fix crash when deleting a location volume actor.

#rnx

diff --git a/Engine/Source/Editor/WorldPartitionEditor/Private/WorldPartition/SWorldPartitionEditorGrid2D.cpp b/Engine/Source/Editor/WorldPartitionEditor/Private/WorldPartition/SWorldPartitionEditorGrid2D.cpp
--- a/Engine/Source/Editor/WorldPartitionEditor/Private/WorldPartition/SWorldPartitionEditorGrid2D.cpp
+++ b/Engine/Source/Editor/WorldPartitionEditor/Private/WorldPartition/SWorldPartitionEditorGrid2D.cpp
@@ -683,6 +683,11 @@ void SWorldPartitionEditorGrid2D::Tick(const FGeometry& AllottedGeometry, const
 	FLoaderInterfaceSet LastHoveredLoaderInterfaces = MoveTemp(HoveredLoaderInterfaces);
 	for (const FLoaderInterface& LoaderInterface : ShownLoaderInterfaces)
 	{
+		if (!LoaderInterface.IsValid())
+		{
+			continue;
+		}
+
 		const IWorldPartitionActorLoaderInterface::ILoaderAdapter* LoaderAdapter = LoaderInterface->GetLoaderAdapter();
 
 		if (!LoaderAdapter->GetBoundingBox()->IsInsideXY(ViewRectWorld))
@@ -809,6 +814,11 @@ uint32 SWorldPartitionEditorGrid2D::PaintActors(const FGeometry& AllottedGeometr
 
 		for (const FLoaderInterface& LoaderInterface : ShownLoaderInterfaces)
 		{
+			if (!LoaderInterface.IsValid())
+			{
+				continue;
+			}
+
 			const IWorldPartitionActorLoaderInterface::ILoaderAdapter* LoaderAdapter = LoaderInterface->GetLoaderAdapter();
 
 			if (LoaderAdapter->GetBoundingBox().IsSet())
