World Partition: include transient and newly created, unsaved actors in world partition editor for convenience.

#rnx

diff --git a/Engine/Source/Editor/WorldPartitionEditor/Private/WorldPartition/SWorldPartitionEditorGrid2D.cpp b/Engine/Source/Editor/WorldPartitionEditor/Private/WorldPartition/SWorldPartitionEditorGrid2D.cpp
--- a/Engine/Source/Editor/WorldPartitionEditor/Private/WorldPartition/SWorldPartitionEditorGrid2D.cpp
+++ b/Engine/Source/Editor/WorldPartitionEditor/Private/WorldPartition/SWorldPartitionEditorGrid2D.cpp
@@ -668,10 +668,10 @@ void SWorldPartitionEditorGrid2D::Tick(const FGeometry& AllottedGeometry, const
 	});
 
 	// Also include transient actor loader adapters that might have been spawned by blutilities, etc. Since these actors can't be saved because they are transient,
-	// they will never get an actor descriptor so they will never appear in the world partition editor.
+	// they will never get an actor descriptor so they will never appear in the world partition editor. Also include unsaved, newly created actors for convenience.
 	for (TActorIterator<AActor> ActorIt(World); ActorIt; ++ActorIt)
 	{
-		if (AActor* Actor = *ActorIt; Actor && Actor->HasAnyFlags(RF_Transient) && Actor->Implements<UWorldPartitionActorLoaderInterface>())
+		if (AActor* Actor = *ActorIt; Actor && !WorldPartition->GetActorDesc(Actor->GetActorGuid()) && Actor->Implements<UWorldPartitionActorLoaderInterface>())
 		{
 			if (IWorldPartitionActorLoaderInterface::ILoaderAdapter* LoaderAdapter = Cast<IWorldPartitionActorLoaderInterface>(Actor)->GetLoaderAdapter())
 			{
