World Partition: cap the location volume height to the world extent height when converting a loading region to a location volume.

#rnx

diff --git a/Engine/Source/Editor/WorldPartitionEditor/Private/WorldPartition/SWorldPartitionEditorGrid2D.cpp b/Engine/Source/Editor/WorldPartitionEditor/Private/WorldPartition/SWorldPartitionEditorGrid2D.cpp
--- a/Engine/Source/Editor/WorldPartitionEditor/Private/WorldPartition/SWorldPartitionEditorGrid2D.cpp
+++ b/Engine/Source/Editor/WorldPartitionEditor/Private/WorldPartition/SWorldPartitionEditorGrid2D.cpp
@@ -339,6 +339,7 @@ void SWorldPartitionEditorGrid2D::ConvertSelectedRegionsToActors()
 	FLoaderInterfaceSet TmpSelectedLoaderInterfaces = MoveTemp(SelectedLoaderInterfaces);
 	ClearSelection();
 
+	const FBox WorldBounds = WorldPartition->GetRuntimeWorldBounds();
 	for (const TWeakInterfacePtr<IWorldPartitionActorLoaderInterface>& SelectedLoaderAdapter : TmpSelectedLoaderInterfaces)
 	{
 		if (UWorldPartitionEditorLoaderAdapter* EditorLoaderAdapter = Cast<UWorldPartitionEditorLoaderAdapter>(SelectedLoaderAdapter.Get()))
@@ -346,8 +347,9 @@ void SWorldPartitionEditorGrid2D::ConvertSelectedRegionsToActors()
 			IWorldPartitionActorLoaderInterface::ILoaderAdapter* LoaderAdapter = EditorLoaderAdapter->GetLoaderAdapter();
 			
 			const FBox LoaderVolumeBox(*LoaderAdapter->GetBoundingBox());
+			const FBox ActorVolumeBox(FVector(LoaderVolumeBox.Min.X, LoaderVolumeBox.Min.Y, WorldBounds.Min.Z), FVector(LoaderVolumeBox.Max.X, LoaderVolumeBox.Max.Y, WorldBounds.Max.Z));
 
-			ALocationVolume* LocationVolume = World->SpawnActor<ALocationVolume>(LoaderVolumeBox.GetCenter(), FRotator::ZeroRotator);
+			ALocationVolume* LocationVolume = World->SpawnActor<ALocationVolume>(ActorVolumeBox.GetCenter(), FRotator::ZeroRotator);
 
 			UCubeBuilder* Builder = NewObject<UCubeBuilder>();
 			Builder->X = 1.0f;
@@ -355,7 +357,7 @@ void SWorldPartitionEditorGrid2D::ConvertSelectedRegionsToActors()
 			Builder->Z = 1.0f;
 			UActorFactory::CreateBrushForVolumeActor(LocationVolume, Builder);
 
-			LocationVolume->GetRootComponent()->SetWorldScale3D(LoaderVolumeBox.GetSize());
+			LocationVolume->GetRootComponent()->SetWorldScale3D(ActorVolumeBox.GetSize());
 
 			LocationVolume->GetLoaderAdapter()->Load();
 
