World Partition Editor: display actor name when hovering on the edge of an actor descriptor bounds.

#rnx

diff --git a/Engine/Source/Editor/WorldPartitionEditor/Private/WorldPartition/SWorldPartitionEditorGrid2D.cpp b/Engine/Source/Editor/WorldPartitionEditor/Private/WorldPartition/SWorldPartitionEditorGrid2D.cpp
--- a/Engine/Source/Editor/WorldPartitionEditor/Private/WorldPartition/SWorldPartitionEditorGrid2D.cpp
+++ b/Engine/Source/Editor/WorldPartitionEditor/Private/WorldPartition/SWorldPartitionEditorGrid2D.cpp
@@ -34,11 +34,17 @@ static bool IsBoundsSelected(const FBox& SelectBox, const FBox& Bounds)
 	return SelectBox.IsValid && Bounds.IntersectXY(SelectBox) && !Bounds.IsInsideXY(SelectBox);
 }
 
-static bool IsBoundsHovered(FVector2D Point, const FBox& Bounds)
+static bool IsBoundsHovered(const FVector2D& Point, const FBox& Bounds)
 {
 	return Bounds.IsInsideOrOnXY(FVector(Point.X, Point.Y, 0));
 }
 
+static bool IsBoundsEdgeHovered(const FVector2D& Point, const FBox2D& Bounds, float Size)
+{
+	const float DistanceToPoint = FMath::Sqrt(Bounds.ExpandBy(-Size * 0.5f).ComputeSquaredDistanceToPoint(Point));
+	return DistanceToPoint > UE_KINDA_SMALL_NUMBER && DistanceToPoint < Size;
+}
+
 template <class T>
 void ForEachIntersectingLoaderAdapters(UWorldPartition* WorldPartition, const FBox& SelectBox, T Func)
 {
@@ -779,28 +785,44 @@ uint32 SWorldPartitionEditorGrid2D::PaintActors(const FGeometry& AllottedGeometr
 		}
 	}
 
-	auto DrawActorLabel = [&OutDrawElements, &LayerId, &AllottedGeometry](const FString& ActorLabel, const FBox2D& ActorViewBox, const FPaintGeometry& ActorGeometry, const FLinearColor& Color, const FSlateFontInfo& Font)
+	auto DrawActorLabel = [this, &OutDrawElements, &LayerId, &AllottedGeometry](const FString& Label, const FVector2D& Pos, const FPaintGeometry& Geometry, const FLinearColor& Color, const FSlateFontInfo& Font, bool bShowBackground)
 	{
-		const FVector2D LabelTextSize = FSlateApplication::Get().GetRenderer()->GetFontMeasureService()->Measure(ActorLabel, Font);
+		const FVector2D LabelTextSize = FSlateApplication::Get().GetRenderer()->GetFontMeasureService()->Measure(Label, Font);
 
 		if (LabelTextSize.X > 0)
 		{
-			const FVector2D ActorViewBoxCenter = ActorViewBox.GetCenter();
-			const FVector2D LabelTextPos = ActorViewBoxCenter - LabelTextSize * 0.5f;
-			const float LabelColorGradient = FMath::Clamp(ActorGeometry.GetLocalSize().X / LabelTextSize.X - 1.0f, 0.0f, 1.0f);
+			const FVector2D BackgroundGrowSize(2.0f);
+			const FVector2D LabelTextPos = Pos - LabelTextSize * 0.5f;
+
+			if (bShowBackground)
+			{
+				const FSlateColorBrush BackgroundBrush(FLinearColor::Black);
+				const FLinearColor LabelBackgroundColor(0, 0, 0, 0.25f);
+
+				FSlateDrawElement::MakeBox(
+					OutDrawElements,
+					++LayerId,
+					AllottedGeometry.ToPaintGeometry(LabelTextPos - BackgroundGrowSize, LabelTextSize + BackgroundGrowSize),
+					&BackgroundBrush,
+					ESlateDrawEffect::None,
+					LabelBackgroundColor
+				);
+			}
+
+			const float LabelColorGradient = bShowBackground ? 1.0f : FMath::Clamp(Geometry.GetLocalSize().X / LabelTextSize.X - 1.0f, 0.0f, 1.0f);
 
 			if (LabelColorGradient > 0.0f)
 			{
-				const FLinearColor LabelColor(Color.R, Color.G, Color.B, Color.A * LabelColorGradient);
+				const FLinearColor LabelForegroundColor(1.0f, 1.0f, 1.0f, Color.A * LabelColorGradient);
 
 				FSlateDrawElement::MakeText(
 					OutDrawElements,
 					++LayerId,
 					AllottedGeometry.ToPaintGeometry(LabelTextPos, FVector2D(1,1)),
-					ActorLabel,
+					Label,
 					Font,
 					ESlateDrawEffect::None,
-					LabelColor
+					LabelForegroundColor
 				);
 			}
 		}
@@ -899,7 +921,7 @@ uint32 SWorldPartitionEditorGrid2D::PaintActors(const FGeometry& AllottedGeometr
 						{
 							const FString ActorLabel = *LoaderAdapter->GetLabel();
 							const FLinearColor LabelColor(1.0f, 1.0f, 1.0f, LoaderColorGradient * FullScreenColorGradient);
-							DrawActorLabel(ActorLabel, ActorViewBox, ActorGeometry, LabelColor, SmallLayoutFont);
+							DrawActorLabel(ActorLabel, ActorViewBox.GetCenter(), ActorGeometry, LabelColor, SmallLayoutFont, false);
 						}
 					}
 				}
@@ -914,6 +936,7 @@ uint32 SWorldPartitionEditorGrid2D::PaintActors(const FGeometry& AllottedGeometr
 		TArray<FVector2D> LinePoints;
 		LinePoints.SetNum(5);
 
+		bool bIsBoundsEdgeHoveredFound = false;
 		for (const FWorldPartitionActorDescViewBoundsProxy& ActorDescView: ActorDescList)
 		{
 			const FBox ActorBounds = ActorDescView.GetBounds();
@@ -936,25 +959,28 @@ uint32 SWorldPartitionEditorGrid2D::PaintActors(const FGeometry& AllottedGeometr
 			const float AreaFadeDistance = 128.0f;
 			if ((Extent.Size2D() < KINDA_SMALL_NUMBER) || (ActorViewBox.GetArea() > MinimumAreaCull))
 			{
+				const UClass* ActorClass = ActorDescView.GetActorNativeClass();
+				const AActor* Actor = ActorDescView.GetActor();
 				const FPaintGeometry ActorGeometry = AllottedGeometry.ToPaintGeometry(TopLeft, BottomRight - TopLeft);
-
+				const bool bIsSelected = Actor ? Actor->IsSelected() : false;
+				const bool bIsBoundsEdgeHovered = !bIsBoundsEdgeHoveredFound && IsBoundsEdgeHovered(MouseCursorPos, ActorViewBox, 10.0f);
 				const float ActorColorGradient = FMath::Min((ActorViewBox.GetArea() - MinimumAreaCull) / AreaFadeDistance, 1.0f);
 				const float ActorBrightness = ActorDescView.GetIsSpatiallyLoaded() ? 1.0f : 0.3f;
-				FLinearColor ActorColor(ActorBrightness, ActorBrightness, ActorBrightness, ActorColorGradient);
 
-				UClass* ActorClass = ActorDescView.GetActorNativeClass();
-				const AActor* Actor = ActorDescView.GetActor();
+				// Only draw the first edge bounds hovered
+				bIsBoundsEdgeHoveredFound |= bIsBoundsEdgeHovered;
+				
+				FLinearColor ActorColor(ActorBrightness, ActorBrightness, ActorBrightness, bIsBoundsEdgeHovered ? 1.0f : ActorColorGradient);
 
-				const bool bIsSelected = Actor ? Actor->IsSelected() : false;
-				if (bIsSelected)
+				if (bIsSelected || bIsBoundsEdgeHovered)
 				{
-					ActorColor = FLinearColor::Yellow;
-
 					const FName ActorLabel = ActorDescView.GetActorLabel();
 					if (!ActorLabel.IsNone())
 					{
-						DrawActorLabel(ActorLabel.ToString(), ActorViewBox, ActorGeometry, ActorColor, SmallLayoutFont);
+						DrawActorLabel(ActorLabel.ToString(), bIsBoundsEdgeHovered ? MouseCursorPos : ActorViewBox.GetCenter(), ActorGeometry, FLinearColor::Yellow, SmallLayoutFont, true);
 					}
+
+					ActorColor = FLinearColor::Yellow;
 				}
 				else if ((SelectBox.GetVolume() > 0) && SelectBox.Intersect(ActorDescView.GetBounds()))
 				{
