World Partition Editor Improvements:
- Greatly optimized rendering dirty and selected actor bounding volumes.
- Added hidden option to display Tick and Paint profiling timers on screen.

Tick time down from ~4ms to ~3ms for a large map.
Paint time down from ~16ms to ~300us for a large map.

#rnx

diff --git a/Engine/Source/Editor/WorldPartitionEditor/Private/WorldPartition/SWorldPartitionEditorGrid2D.cpp b/Engine/Source/Editor/WorldPartitionEditor/Private/WorldPartition/SWorldPartitionEditorGrid2D.cpp
--- a/Engine/Source/Editor/WorldPartitionEditor/Private/WorldPartition/SWorldPartitionEditorGrid2D.cpp
+++ b/Engine/Source/Editor/WorldPartitionEditor/Private/WorldPartition/SWorldPartitionEditorGrid2D.cpp
@@ -89,15 +89,19 @@ void ForEachIntersectingLoaderAdapters(UWorldPartition* WorldPartition, const FB
 class FWorldPartitionActorDescViewBoundsProxy : public FWorldPartitionActorDescView
 {
 public:
-	FWorldPartitionActorDescViewBoundsProxy(const FWorldPartitionActorDesc* InActorDesc)
+	FWorldPartitionActorDescViewBoundsProxy(const FWorldPartitionActorDesc* InActorDesc, bool bInUseActor)
 		: FWorldPartitionActorDescView(InActorDesc)
+		, bUseActor(bInUseActor)
 	{}
 
 	FBox GetBounds() const
 	{
-		if (AActor* Actor = GetActor())
+		if (bUseActor)
 		{
-			return Actor->GetStreamingBounds();
+			if (AActor* Actor = GetActor())
+			{
+				return Actor->GetStreamingBounds();
+			}
 		}
 
 		return ActorDesc->GetBounds();
@@ -105,8 +109,10 @@ public:
 
 	AActor* GetActor() const
 	{
-		return ActorDesc->GetActor(false);
+		return bUseActor ? ActorDesc->GetActor(false) : nullptr;
 	}
+
+	bool bUseActor;
 };
 
 SWorldPartitionEditorGrid2D::FEditorCommands::FEditorCommands()
@@ -137,7 +143,10 @@ SWorldPartitionEditorGrid2D::SWorldPartitionEditorGrid2D()
 	, bIsDragSelecting(false)
 	, bIsPanning(false)
 	, bShowActors(false)
+	, bShowProfiling(false)
 	, SelectBox(ForceInit)
+	, TickTime(0)
+	, PaintTime(0)
 {
 	FEditorCommands::Register();
 	
@@ -641,16 +650,17 @@ int32 SWorldPartitionEditorGrid2D::PaintGrid(const FGeometry& AllottedGeometry,
 
 void SWorldPartitionEditorGrid2D::Tick(const FGeometry& AllottedGeometry, const double InCurrentTime, const float InDeltaTime)
 {
+	double LocalTickTime = -FPlatformTime::Seconds();
+
 	const FBox2D ViewRect(FVector2D(ForceInitToZero), AllottedGeometry.GetLocalSize());
 	const FBox2D WorldViewRect(ScreenToWorld.TransformPoint(ViewRect.Min), ScreenToWorld.TransformPoint(ViewRect.Max));
 	const FBox ViewRectWorld(FVector(WorldViewRect.Min.X, WorldViewRect.Min.Y, -HALF_WORLD_MAX), FVector(WorldViewRect.Max.X, WorldViewRect.Max.Y, HALF_WORLD_MAX));
 
 	ShownActorGuids.Reset();
+	DirtyActorGuids.Reset();
 	ShownLoaderInterfaces.Reset();
 	HighlightedLoaderInterfaces.Reset();
 	
-	double MinLoaderDistance = DBL_MAX;
-	double MinActorDistance = DBL_MAX;
 	for (UWorldPartitionEditorLoaderAdapter* EditorLoaderAdapter : WorldPartition->GetRegisteredEditorLoaderAdapters())
 	{
 		check(EditorLoaderAdapter);
@@ -689,11 +699,19 @@ void SWorldPartitionEditorGrid2D::Tick(const FGeometry& AllottedGeometry, const
 	// they will never get an actor descriptor so they will never appear in the world partition editor. Also include unsaved, newly created actors for convenience.
 	for (TActorIterator<AActor> ActorIt(World); ActorIt; ++ActorIt)
 	{
-		if (AActor* Actor = *ActorIt; Actor && !WorldPartition->GetActorDesc(Actor->GetActorGuid()) && Actor->Implements<UWorldPartitionActorLoaderInterface>())
+		if (AActor* Actor = *ActorIt)
 		{
-			if (IWorldPartitionActorLoaderInterface::ILoaderAdapter* LoaderAdapter = Cast<IWorldPartitionActorLoaderInterface>(Actor)->GetLoaderAdapter())
+			if (!WorldPartition->GetActorDesc(Actor->GetActorGuid()) && Actor->Implements<UWorldPartitionActorLoaderInterface>())
 			{
-				ShownLoaderInterfaces.Add(Actor);
+				if (IWorldPartitionActorLoaderInterface::ILoaderAdapter* LoaderAdapter = Cast<IWorldPartitionActorLoaderInterface>(Actor)->GetLoaderAdapter())
+				{
+					ShownLoaderInterfaces.Add(Actor);
+				}
+			}
+
+			if(Actor->IsSelected() || Actor->GetPackage()->IsDirty())
+			{
+				DirtyActorGuids.Add(Actor->GetActorGuid());
 			}
 		}
 	}
@@ -770,6 +788,9 @@ void SWorldPartitionEditorGrid2D::Tick(const FGeometry& AllottedGeometry, const
 			ShownActorGuids.Add(Actor->GetActorGuid());
 		}
 	}
+
+	LocalTickTime += FPlatformTime::Seconds();
+	TickTime = TickTime * 0.95 + LocalTickTime * 0.05;
 }
 
 uint32 SWorldPartitionEditorGrid2D::PaintActors(const FGeometry& AllottedGeometry, const FSlateRect& MyCullingRect, FSlateWindowElementList& OutDrawElements, uint32 LayerId) const
@@ -791,7 +812,7 @@ uint32 SWorldPartitionEditorGrid2D::PaintActors(const FGeometry& AllottedGeometr
 	{
 		if (FWorldPartitionActorDesc* ActorDesc = WorldPartition->GetActorDesc(ActorGuid))
 		{
-			ActorDescList.Emplace(ActorDesc);
+			ActorDescList.Emplace(ActorDesc, DirtyActorGuids.Contains(ActorGuid));
 		}
 	}
 
@@ -1012,7 +1033,7 @@ uint32 SWorldPartitionEditorGrid2D::PaintActors(const FGeometry& AllottedGeometr
 	return LayerId + 1;
 }
 
-uint32 SWorldPartitionEditorGrid2D::PaintScaleRuler(const FGeometry& AllottedGeometry, const FSlateRect& MyCullingRect, FSlateWindowElementList& OutDrawElements, uint32 LayerId) const
+uint32 SWorldPartitionEditorGrid2D::PaintTextInfo(const FGeometry& AllottedGeometry, const FSlateRect& MyCullingRect, FSlateWindowElementList& OutDrawElements, uint32 LayerId) const
 {
 	const float	ScaleRulerLength = 100.f; // pixels
 	TArray<FVector2D> LinePoints;
@@ -1054,7 +1075,7 @@ uint32 SWorldPartitionEditorGrid2D::PaintScaleRuler(const FGeometry& AllottedGeo
 	const FBox WorldBounds = WorldPartition->GetRuntimeWorldBounds();
 	const FVector WorldBoundsExtentInKM = (WorldBounds.GetExtent() * 2.0f) / 100000.0f;
 	RulerText = FString::Printf(TEXT("%.2fx%.2fx%.2f km"), WorldBoundsExtentInKM.X, WorldBoundsExtentInKM.Y, WorldBoundsExtentInKM.Z);
-	
+
 	FSlateDrawElement::MakeText(
 		OutDrawElements,
 		LayerId,
@@ -1064,6 +1085,31 @@ uint32 SWorldPartitionEditorGrid2D::PaintScaleRuler(const FGeometry& AllottedGeo
 		ESlateDrawEffect::None,
 		FLinearColor::White);
 
+	// Show profiling
+	if (bShowProfiling)
+	{
+		RulerText = FString::Printf(TEXT("TickTime=%s"), *FPlatformTime::PrettyTime(TickTime));
+		FSlateDrawElement::MakeText(
+			OutDrawElements,
+			LayerId,
+			AllottedGeometry.ToOffsetPaintGeometry(FVector2D(10, 107)),
+			RulerText,
+			FAppStyle::GetFontStyle("NormalFont"),
+			ESlateDrawEffect::None,
+			FLinearColor::Gray);
+
+		const FVector2D TextSize = FSlateApplication::Get().GetRenderer()->GetFontMeasureService()->Measure(RulerText, FAppStyle::GetFontStyle("NormalFont"));
+
+		RulerText = FString::Printf(TEXT("PaintTime=%s"), *FPlatformTime::PrettyTime(PaintTime));
+		FSlateDrawElement::MakeText(
+			OutDrawElements,
+			LayerId,
+			AllottedGeometry.ToOffsetPaintGeometry(FVector2D(10, 107 + TextSize.Y + 2)),
+			RulerText,
+			FAppStyle::GetFontStyle("NormalFont"),
+			ESlateDrawEffect::None,
+			FLinearColor::Gray);
+	}
 		
 	return LayerId + 1;
 }
@@ -1181,6 +1227,8 @@ uint32 SWorldPartitionEditorGrid2D::PaintSelection(const FGeometry& AllottedGeom
 
 int32 SWorldPartitionEditorGrid2D::OnPaint(const FPaintArgs& Args, const FGeometry& AllottedGeometry, const FSlateRect& MyCullingRect, FSlateWindowElementList& OutDrawElements, int32 LayerId, const FWidgetStyle& InWidgetStyle, bool bParentEnabled) const
 {
+	double LocalPaintTime = -FPlatformTime::Seconds();
+
 	if (WorldPartition)
 	{
 		const bool bResetView = !ScreenRect.bIsValid;
@@ -1196,7 +1244,7 @@ int32 SWorldPartitionEditorGrid2D::OnPaint(const FPaintArgs& Args, const FGeomet
 
 		LayerId = PaintGrid(AllottedGeometry, MyCullingRect, OutDrawElements, ++LayerId);
 		LayerId = PaintActors(AllottedGeometry, MyCullingRect, OutDrawElements, ++LayerId);
-		LayerId = PaintScaleRuler(AllottedGeometry, MyCullingRect, OutDrawElements, ++LayerId);
+		LayerId = PaintTextInfo(AllottedGeometry, MyCullingRect, OutDrawElements, ++LayerId);
 		LayerId = PaintViewer(AllottedGeometry, MyCullingRect, OutDrawElements, ++LayerId);
 		LayerId = PaintSelection(AllottedGeometry, MyCullingRect, OutDrawElements, ++LayerId);
 		LayerId = PaintSoftwareCursor(AllottedGeometry, MyCullingRect, OutDrawElements, ++LayerId);
@@ -1213,7 +1261,10 @@ int32 SWorldPartitionEditorGrid2D::OnPaint(const FPaintArgs& Args, const FGeomet
 		}
 	}
 
-	return SWorldPartitionEditorGrid::OnPaint(Args, AllottedGeometry, MyCullingRect, OutDrawElements, LayerId, InWidgetStyle, bParentEnabled);
+	int32 Result = SWorldPartitionEditorGrid::OnPaint(Args, AllottedGeometry, MyCullingRect, OutDrawElements, LayerId, InWidgetStyle, bParentEnabled);
+	LocalPaintTime += FPlatformTime::Seconds();
+	PaintTime = PaintTime * 0.95 + LocalPaintTime * 0.05;
+	return Result;
 }
 
 int32 SWorldPartitionEditorGrid2D::PaintSoftwareCursor(const FGeometry& AllottedGeometry, const FSlateRect& MyCullingRect, FSlateWindowElementList& OutDrawElements, int32 LayerId) const
diff --git a/Engine/Source/Editor/WorldPartitionEditor/Public/WorldPartition/SWorldPartitionEditorGrid2D.h b/Engine/Source/Editor/WorldPartitionEditor/Public/WorldPartition/SWorldPartitionEditorGrid2D.h
--- a/Engine/Source/Editor/WorldPartitionEditor/Public/WorldPartition/SWorldPartitionEditorGrid2D.h
+++ b/Engine/Source/Editor/WorldPartitionEditor/Public/WorldPartition/SWorldPartitionEditorGrid2D.h
@@ -54,7 +54,7 @@ public:
 	virtual int32 OnPaint(const FPaintArgs& Args, const FGeometry& AllottedGeometry, const FSlateRect& MyCullingRect, FSlateWindowElementList& OutDrawElements, int32 LayerId, const FWidgetStyle& InWidgetStyle, bool bParentEnabled) const override;
 	virtual int32 PaintGrid(const FGeometry& AllottedGeometry, const FSlateRect& MyCullingRect, FSlateWindowElementList& OutDrawElements, int32 LayerId) const;
 	virtual uint32 PaintActors(const FGeometry& AllottedGeometry, const FSlateRect& MyCullingRect, FSlateWindowElementList& OutDrawElements, uint32 LayerId) const;
-	virtual uint32 PaintScaleRuler(const FGeometry& AllottedGeometry, const FSlateRect& MyCullingRect, FSlateWindowElementList& OutDrawElements, uint32 LayerId) const;
+	virtual uint32 PaintTextInfo(const FGeometry& AllottedGeometry, const FSlateRect& MyCullingRect, FSlateWindowElementList& OutDrawElements, uint32 LayerId) const;
 	virtual uint32 PaintViewer(const FGeometry& AllottedGeometry, const FSlateRect& MyCullingRect, FSlateWindowElementList& OutDrawElements, uint32 LayerId) const;
 	virtual uint32 PaintSelection(const FGeometry& AllottedGeometry, const FSlateRect& MyCullingRect, FSlateWindowElementList& OutDrawElements, uint32 LayerId) const;
 	virtual int32 PaintSoftwareCursor(const FGeometry& AllottedGeometry, const FSlateRect& MyCullingRect, FSlateWindowElementList& OutDrawElements, int32 LayerId) const;
@@ -91,6 +91,7 @@ protected:
 	bool bIsDragSelecting;
 	bool bIsPanning;
 	bool bShowActors;
+	bool bShowProfiling;
 	FVector2D MouseCursorPos;
 	FVector2D MouseCursorPosWorld;
 	FVector2D LastMouseCursorPosWorldDrag;
@@ -126,9 +127,14 @@ protected:
 	
 	// Updated every tick
 	TSet<FGuid> ShownActorGuids;
+	TSet<FGuid> DirtyActorGuids;
 	FLoaderInterfaceSet ShownLoaderInterfaces;
 	FLoaderInterfaceSet HighlightedLoaderInterfaces;
 	FLoaderInterfaceSet HoveredLoaderInterfaces;
 	FLoaderInterfaceStack HoveredLoaderInterfacesStack;
 	FLoaderInterface HoveredLoaderInterface;
+
+	// Profiling
+	mutable double TickTime;
+	mutable double PaintTime;
 };
\ No newline at end of file
