World Partition: fix random crash when unloading loading regions.

#rnx

diff --git a/Engine/Source/Editor/WorldPartitionEditor/Private/WorldPartition/SWorldPartitionEditorGrid2D.cpp b/Engine/Source/Editor/WorldPartitionEditor/Private/WorldPartition/SWorldPartitionEditorGrid2D.cpp
--- a/Engine/Source/Editor/WorldPartitionEditor/Private/WorldPartition/SWorldPartitionEditorGrid2D.cpp
+++ b/Engine/Source/Editor/WorldPartitionEditor/Private/WorldPartition/SWorldPartitionEditorGrid2D.cpp
@@ -709,7 +709,7 @@ void SWorldPartitionEditorGrid2D::Tick(const FGeometry& AllottedGeometry, const
 				}
 			}
 
-			if(Actor->IsSelected() || Actor->GetPackage()->IsDirty())
+			if(Actor->GetPackage()->IsDirty())
 			{
 				DirtyActorGuids.Add(Actor->GetActorGuid());
 			}
@@ -764,7 +764,8 @@ void SWorldPartitionEditorGrid2D::Tick(const FGeometry& AllottedGeometry, const
 			{
 				if (FLoaderInterface StackedHoveredLoaderInterface = HoveredLoaderInterfacesStack.Pop(); StackedHoveredLoaderInterface.IsValid())
 				{
-					if (IsBoundsHovered(MouseCursorPosWorld, *StackedHoveredLoaderInterface->GetLoaderAdapter()->GetBoundingBox()))
+					IWorldPartitionActorLoaderInterface::ILoaderAdapter* LoaderAdapter = StackedHoveredLoaderInterface->GetLoaderAdapter();
+					if (LoaderAdapter && IsBoundsHovered(MouseCursorPosWorld, *StackedHoveredLoaderInterface->GetLoaderAdapter()->GetBoundingBox()))
 					{
 						HoveredLoaderInterface = StackedHoveredLoaderInterface;
 						break;
