World Partition: fix code that wasn't supposed to be submitted when 21399140 went in.

#rnx

diff --git a/Engine/Source/Editor/WorldPartitionEditor/Private/WorldPartition/SWorldPartitionEditorGrid2D.cpp b/Engine/Source/Editor/WorldPartitionEditor/Private/WorldPartition/SWorldPartitionEditorGrid2D.cpp
--- a/Engine/Source/Editor/WorldPartitionEditor/Private/WorldPartition/SWorldPartitionEditorGrid2D.cpp
+++ b/Engine/Source/Editor/WorldPartitionEditor/Private/WorldPartition/SWorldPartitionEditorGrid2D.cpp
@@ -658,7 +658,7 @@ void SWorldPartitionEditorGrid2D::Tick(const FGeometry& AllottedGeometry, const
 
 	ShownActorGuids.Reset();
 	DirtyActorGuids.Reset();
-	ShownLoaderInterfaces.Reset();
+	ShownLoaderInterfaces.Reset(); 
 	HighlightedLoaderInterfaces.Reset();
 	
 	for (UWorldPartitionEditorLoaderAdapter* EditorLoaderAdapter : WorldPartition->GetRegisteredEditorLoaderAdapters())
@@ -709,7 +709,7 @@ void SWorldPartitionEditorGrid2D::Tick(const FGeometry& AllottedGeometry, const
 				}
 			}
 
-			if(Actor->GetPackage()->IsDirty())
+			if(Actor->IsSelected() || Actor->GetPackage()->IsDirty())
 			{
 				DirtyActorGuids.Add(Actor->GetActorGuid());
 			}
