World Partition: added the possibility to change the associated runtime hash through the UI.

#rnx

diff --git a/Engine/Source/Editor/WorldPartitionEditor/Private/WorldPartition/Customizations/WorldPartitionDetailsCustomization.cpp b/Engine/Source/Editor/WorldPartitionEditor/Private/WorldPartition/Customizations/WorldPartitionDetailsCustomization.cpp
--- a/Engine/Source/Editor/WorldPartitionEditor/Private/WorldPartition/Customizations/WorldPartitionDetailsCustomization.cpp
+++ b/Engine/Source/Editor/WorldPartitionEditor/Private/WorldPartition/Customizations/WorldPartitionDetailsCustomization.cpp
@@ -3,6 +3,7 @@
 #include "WorldPartitionDetailsCustomization.h"
 #include "WorldPartition/WorldPartition.h"
 #include "WorldPartition/WorldPartitionRuntimeHash.h"
+#include "PropertyCustomizationHelpers.h"
 #include "DetailLayoutBuilder.h"
 #include "DetailCategoryBuilder.h"
 #include "DetailWidgetRow.h"
@@ -22,16 +23,19 @@ TSharedRef<IDetailCustomization> FWorldPartitionDetails::MakeInstance()
 	return MakeShareable(new FWorldPartitionDetails);
 }
 
-void FWorldPartitionDetails::CustomizeDetails(IDetailLayoutBuilder& DetailBuilder)
+void FWorldPartitionDetails::CustomizeDetails(IDetailLayoutBuilder& InDetailBuilder)
 {
 	TArray<TWeakObjectPtr<UObject>> ObjectsBeingCustomized;
-	DetailBuilder.GetObjectsBeingCustomized(ObjectsBeingCustomized);
+	InDetailBuilder.GetObjectsBeingCustomized(ObjectsBeingCustomized);
 	check(ObjectsBeingCustomized.Num() == 1);
 
+	DetailBuilder = &InDetailBuilder;
 	WorldPartition = CastChecked<UWorldPartition>(ObjectsBeingCustomized[0].Get());
+	RuntimeHashClass = (WorldPartition.IsValid() && WorldPartition->RuntimeHash) ? WorldPartition->RuntimeHash->GetClass() : nullptr;
 
-	IDetailCategoryBuilder& WorldPartitionCategory = DetailBuilder.EditCategory("WorldPartition");
+	IDetailCategoryBuilder& WorldPartitionCategory = InDetailBuilder.EditCategory("WorldPartition");
 
+	// Enable streaming checkbox
 	WorldPartitionCategory.AddCustomRow(LOCTEXT("EnableStreaming", "Enable Streaming"), false)
 		.NameContent()
 		[
@@ -48,6 +52,27 @@ void FWorldPartitionDetails::CustomizeDetails(IDetailLayoutBuilder& DetailBuilde
 		]
 		.Visibility(TAttribute<EVisibility>::CreateLambda([this]() { return WorldPartition.IsValid() && WorldPartition->SupportsStreaming() ? EVisibility::Visible : EVisibility::Hidden; }));
 
+	// Runtime hash class selector
+	WorldPartitionCategory.AddCustomRow(LOCTEXT("RuntimeHashClass", "Runtime Hash Class"), false)
+		.NameContent()
+		[
+			SNew(STextBlock)
+			.Text(LOCTEXT("WorldPartitionRuntimeHashClass", "Runtime Hash Class"))
+			.ToolTipText(LOCTEXT("WorldPartitionRuntimeHashClass_ToolTip", "Set the world partition runtime hash class."))
+			.Font(IDetailLayoutBuilder::GetDetailFont())
+		]
+		.ValueContent()
+		[
+			SNew(SClassPropertyEntryBox)
+			.MetaClass(UWorldPartitionRuntimeHash::StaticClass())
+			.AllowNone(false)
+			.HideViewOptions(true)
+			.SelectedClass_Lambda([this]() { return RuntimeHashClass; })
+			.OnSetClass_Lambda([this](const UClass* Class) { HandleWorldPartitionRuntimeHashClassChanged(Class); })
+		]
+		.Visibility(TAttribute<EVisibility>::CreateLambda([this]() { return WorldPartition.IsValid() && WorldPartition->IsStreamingEnabled() ? EVisibility::Visible : EVisibility::Hidden; }));
+
+	// Runtime hash properties
 	if (WorldPartition->RuntimeHash)
 	{
 		FAddPropertyParams Params;
@@ -64,9 +89,9 @@ void FWorldPartitionDetails::CustomizeDetails(IDetailLayoutBuilder& DetailBuilde
 	}
 }
 
-void FWorldPartitionDetails::HandleWorldPartitionEnableStreamingChanged(ECheckBoxState CheckState)
+void FWorldPartitionDetails::HandleWorldPartitionEnableStreamingChanged(ECheckBoxState InCheckState)
 {
-	if (CheckState == ECheckBoxState::Checked)
+	if (InCheckState == ECheckBoxState::Checked)
 	{
 		if (WorldPartition->CanBeUsedByLevelInstance())
 		{
@@ -102,4 +127,22 @@ void FWorldPartitionDetails::HandleWorldPartitionEnableStreamingChanged(ECheckBo
 	}
 }
 
+void FWorldPartitionDetails::HandleWorldPartitionRuntimeHashClassChanged(const UClass* InRuntimeHashClass)
+{
+	check(InRuntimeHashClass);
+	check(InRuntimeHashClass->IsChildOf<UWorldPartitionRuntimeHash>());
+	check(!(InRuntimeHashClass->GetClassFlags() & CLASS_Abstract));
+
+	RuntimeHashClass = InRuntimeHashClass;
+
+	if (WorldPartition.IsValid() && (!WorldPartition->RuntimeHash || WorldPartition->RuntimeHash->GetClass() != RuntimeHashClass))
+	{
+		WorldPartition->Modify();
+		WorldPartition->RuntimeHash = NewObject<UWorldPartitionRuntimeHash>(WorldPartition.Get(), RuntimeHashClass, NAME_None, RF_Transactional);
+		WorldPartition->RuntimeHash->SetDefaultValues();
+
+		DetailBuilder->ForceRefreshDetails();
+	}
+}
+
 #undef LOCTEXT_NAMESPACE
diff --git a/Engine/Source/Editor/WorldPartitionEditor/Private/WorldPartition/Customizations/WorldPartitionDetailsCustomization.h b/Engine/Source/Editor/WorldPartitionEditor/Private/WorldPartition/Customizations/WorldPartitionDetailsCustomization.h
--- a/Engine/Source/Editor/WorldPartitionEditor/Private/WorldPartition/Customizations/WorldPartitionDetailsCustomization.h
+++ b/Engine/Source/Editor/WorldPartitionEditor/Private/WorldPartition/Customizations/WorldPartitionDetailsCustomization.h
@@ -23,10 +23,15 @@ protected:
 
 private:
 	/** IDetailCustomization interface */
-	virtual void CustomizeDetails(IDetailLayoutBuilder& DetailBuilder) override;
+	virtual void CustomizeDetails(IDetailLayoutBuilder& InDetailBuilder) override;
 
 	// Callback for changes in the world partition enable streaming.
-	void HandleWorldPartitionEnableStreamingChanged(ECheckBoxState CheckState);
+	void HandleWorldPartitionEnableStreamingChanged(ECheckBoxState InCheckState);
 
+	// Callback for changes in the world partition runtime hash class.
+	void HandleWorldPartitionRuntimeHashClassChanged(const UClass* InRuntimeHashClass);
+
+	IDetailLayoutBuilder* DetailBuilder;
 	TWeakObjectPtr<UWorldPartition> WorldPartition;
+	const UClass* RuntimeHashClass;
 };
\ No newline at end of file
