World Partition Editor QoL: snap camera to first ground collision when using "move camera here".

#rnx

diff --git a/Engine/Source/Editor/WorldPartitionEditor/Private/WorldPartition/SWorldPartitionEditorGrid2D.cpp b/Engine/Source/Editor/WorldPartitionEditor/Private/WorldPartition/SWorldPartitionEditorGrid2D.cpp
--- a/Engine/Source/Editor/WorldPartitionEditor/Private/WorldPartition/SWorldPartitionEditorGrid2D.cpp
+++ b/Engine/Source/Editor/WorldPartitionEditor/Private/WorldPartition/SWorldPartitionEditorGrid2D.cpp
@@ -402,11 +402,21 @@ void SWorldPartitionEditorGrid2D::ConvertSelectedRegionsToActors()
 
 void SWorldPartitionEditorGrid2D::MoveCameraHere()
 {
+	FVector WorldLocation = FVector(MouseCursorPosWorld, 0);
+
+	FHitResult HitResult;
+	const FVector TraceStart(WorldLocation.X, WorldLocation.Y, HALF_WORLD_MAX);
+	const FVector TraceEnd(WorldLocation.X, WorldLocation.Y, -HALF_WORLD_MAX);
+	const FCollisionQueryParams TraceParams(SCENE_QUERY_STAT(SWorldPartitionEditorGrid2D_MoveCameraHere), true);
+	const bool bHitResultValid = World->LineTraceSingleByChannel(HitResult, TraceStart, TraceEnd, ECC_WorldStatic, TraceParams);
+
 	for (FLevelEditorViewportClient* LevelVC : GEditor->GetLevelViewportClients())
 	{
-		const FVector WorldLocation = FVector(MouseCursorPosWorld, LevelVC->GetViewLocation().Z);
+		WorldLocation.Z = bHitResultValid ? HitResult.ImpactPoint.Z + 250.0f : LevelVC->GetViewLocation().Z;
+
 		LevelVC->SetViewLocation(WorldLocation);
 		LevelVC->Invalidate();
+
 		FEditorDelegates::OnEditorCameraMoved.Broadcast(WorldLocation, LevelVC->GetViewRotation(), LevelVC->ViewportType, LevelVC->ViewIndex);
 	}
 }
