World Partition Editor:
- Tweaked the selected loading regions appearance to be clearer to the user.
- Include the monimap volume actor when focusing to world bounds if it exists.
- Clear the current selection when starting a marquee selection without the control key down.
- Added cvar wp.Editor.ToggleShowEditorProfiling to display profie stats.

#rnx

diff --git a/Engine/Source/Editor/WorldPartitionEditor/Private/WorldPartition/SWorldPartitionEditorGrid2D.cpp b/Engine/Source/Editor/WorldPartitionEditor/Private/WorldPartition/SWorldPartitionEditorGrid2D.cpp
--- a/Engine/Source/Editor/WorldPartitionEditor/Private/WorldPartition/SWorldPartitionEditorGrid2D.cpp
+++ b/Engine/Source/Editor/WorldPartitionEditor/Private/WorldPartition/SWorldPartitionEditorGrid2D.cpp
@@ -13,6 +13,7 @@
 #include "WorldPartition/WorldPartitionActorDescView.h"
 #include "WorldPartition/WorldPartitionEditorPerProjectUserSettings.h"
 #include "WorldPartition/WorldPartitionHelpers.h"
+#include "WorldPartition/WorldPartitionMiniMapVolume.h"
 #include "WorldPartition/LoaderAdapter/LoaderAdapterShape.h"
 #include "ActorFactories/ActorFactory.h"
 #include "Builders/CubeBuilder.h"
@@ -30,6 +31,32 @@
 
 #define LOCTEXT_NAMESPACE "WorldPartitionEditor"
 
+static bool GShowEditorProfilingStats = 0;
+static FAutoConsoleCommand CVarToggleShowEditorProfilingStats(
+	TEXT("wp.Editor.ToggleShowEditorProfiling"),
+	TEXT("Toggles showing editor profiling stats."),
+	FConsoleCommandDelegate::CreateLambda([] { GShowEditorProfilingStats = !GShowEditorProfilingStats; })
+);
+
+class FWeightedMovingAverageScope
+{
+public:
+	FORCEINLINE FWeightedMovingAverageScope(double& InValue)
+		: Value(InValue)
+	{
+		Time = -FPlatformTime::Seconds();
+	}
+
+	FORCEINLINE ~FWeightedMovingAverageScope()
+	{
+		Value = FMath::WeightedMovingAverage(Time + FPlatformTime::Seconds(), Value, 0.05);
+	}
+
+private:
+	double Time;
+	double& Value;
+};
+
 static bool IsBoundsSelected(const FBox& SelectBox, const FBox& Bounds)
 {
 	return SelectBox.IsValid && Bounds.IntersectXY(SelectBox) && !Bounds.IsInsideXY(SelectBox);
@@ -143,7 +170,6 @@ SWorldPartitionEditorGrid2D::SWorldPartitionEditorGrid2D()
 	, bIsDragSelecting(false)
 	, bIsPanning(false)
 	, bShowActors(false)
-	, bShowProfiling(false)
 	, SelectBox(ForceInit)
 	, TickTime(0)
 	, PaintTime(0)
@@ -439,6 +465,11 @@ FReply SWorldPartitionEditorGrid2D::OnMouseButtonDown(const FGeometry& MyGeometr
 
 		if (bIsLeftMouseButtonEffecting)
 		{
+			if (!MouseEvent.IsControlDown())
+			{
+				ClearSelection();
+			}
+
 			SelectionStart = MouseCursorPosWorld;
 			SelectionEnd = SelectionStart;
 			SelectBox.Init();
@@ -660,7 +691,7 @@ int32 SWorldPartitionEditorGrid2D::PaintGrid(const FGeometry& AllottedGeometry,
 
 void SWorldPartitionEditorGrid2D::Tick(const FGeometry& AllottedGeometry, const double InCurrentTime, const float InDeltaTime)
 {
-	double LocalTickTime = -FPlatformTime::Seconds();
+	FWeightedMovingAverageScope ProfileMeanValue(TickTime);
 
 	const FBox2D ViewRect(FVector2D(ForceInitToZero), AllottedGeometry.GetLocalSize());
 	const FBox2D WorldViewRect(ScreenToWorld.TransformPoint(ViewRect.Min), ScreenToWorld.TransformPoint(ViewRect.Max));
@@ -669,7 +700,6 @@ void SWorldPartitionEditorGrid2D::Tick(const FGeometry& AllottedGeometry, const
 	ShownActorGuids.Reset();
 	DirtyActorGuids.Reset();
 	ShownLoaderInterfaces.Reset(); 
-	HighlightedLoaderInterfaces.Reset();
 	
 	for (UWorldPartitionEditorLoaderAdapter* EditorLoaderAdapter : WorldPartition->GetRegisteredEditorLoaderAdapters())
 	{
@@ -742,11 +772,6 @@ void SWorldPartitionEditorGrid2D::Tick(const FGeometry& AllottedGeometry, const
 			{
 				HoveredLoaderInterfaces.Add(LoaderInterface);
 			}
-
-			if (SelectedLoaderInterfaces.Contains(LoaderInterface))
-			{
-				HighlightedLoaderInterfaces.Add(LoaderInterface);
-			}
 		}
 	}
 
@@ -799,9 +824,6 @@ void SWorldPartitionEditorGrid2D::Tick(const FGeometry& AllottedGeometry, const
 			ShownActorGuids.Add(Actor->GetActorGuid());
 		}
 	}
-
-	LocalTickTime += FPlatformTime::Seconds();
-	TickTime = TickTime * 0.95 + LocalTickTime * 0.05;
 }
 
 uint32 SWorldPartitionEditorGrid2D::PaintActors(const FGeometry& AllottedGeometry, const FSlateRect& MyCullingRect, FSlateWindowElementList& OutDrawElements, uint32 LayerId) const
@@ -937,8 +959,28 @@ uint32 SWorldPartitionEditorGrid2D::PaintActors(const FGeometry& AllottedGeometr
 
 						// Outline
 						{
-							const bool IsHighlighted = HighlightedLoaderInterfaces.Contains(LoaderInterface) || LocalHoveredLoaderAdapter == LoaderAdapter;
-							const FLinearColor OutlineColor = IsHighlighted ? FLinearColor::Yellow : FLinearColor::White;
+							const bool bIsSelected = SelectedLoaderInterfaces.Contains(LoaderInterface);
+							const bool bIsInsideSelection = (SelectBox.GetVolume() > 0) && SelectBox.Intersect(AdapterBounds);
+							const bool bIsLocalHovered = LocalHoveredLoaderAdapter == LoaderAdapter;
+							
+							float OutlineThickness = 2.0f;
+							FLinearColor OutlineColor = FLinearColor::Gray;
+
+							if (bIsSelected)
+							{
+								OutlineThickness = 4.0f;
+								OutlineColor = FLinearColor::Yellow;
+							}
+							else if (bIsInsideSelection)
+							{
+								OutlineThickness = 4.0f;
+								OutlineColor = FLinearColor::Yellow * 0.5f;
+							}
+							else if (bIsLocalHovered)
+							{
+								OutlineThickness = 4.0f;
+								OutlineColor = FLinearColor::White;
+							}
 
 							LinePoints[0] = TopLeft;
 							LinePoints[1] = TopRight;
@@ -955,7 +997,7 @@ uint32 SWorldPartitionEditorGrid2D::PaintActors(const FGeometry& AllottedGeometr
 								ESlateDrawEffect::None,
 								OutlineColor,
 								true,
-								IsHighlighted ? 4.0f : 2.0f
+								OutlineThickness
 							);
 						}
 
@@ -1097,7 +1139,7 @@ uint32 SWorldPartitionEditorGrid2D::PaintTextInfo(const FGeometry& AllottedGeome
 		FLinearColor::White);
 
 	// Show profiling
-	if (bShowProfiling)
+	if (GShowEditorProfilingStats)
 	{
 		RulerText = FString::Printf(TEXT("TickTime=%s"), *FPlatformTime::PrettyTime(TickTime));
 		FSlateDrawElement::MakeText(
@@ -1238,7 +1280,7 @@ uint32 SWorldPartitionEditorGrid2D::PaintSelection(const FGeometry& AllottedGeom
 
 int32 SWorldPartitionEditorGrid2D::OnPaint(const FPaintArgs& Args, const FGeometry& AllottedGeometry, const FSlateRect& MyCullingRect, FSlateWindowElementList& OutDrawElements, int32 LayerId, const FWidgetStyle& InWidgetStyle, bool bParentEnabled) const
 {
-	double LocalPaintTime = -FPlatformTime::Seconds();
+	FWeightedMovingAverageScope ProfileMeanValue(PaintTime);
 
 	if (WorldPartition)
 	{
@@ -1248,7 +1290,7 @@ int32 SWorldPartitionEditorGrid2D::OnPaint(const FPaintArgs& Args, const FGeomet
 
 		if (bResetView)
 		{
-			FocusBox(WorldPartition->GetRuntimeWorldBounds());
+			const_cast<SWorldPartitionEditorGrid2D*>(this)->FocusSelection();
 		}
 
 		UpdateTransform();
@@ -1272,10 +1314,7 @@ int32 SWorldPartitionEditorGrid2D::OnPaint(const FPaintArgs& Args, const FGeomet
 		}
 	}
 
-	int32 Result = SWorldPartitionEditorGrid::OnPaint(Args, AllottedGeometry, MyCullingRect, OutDrawElements, LayerId, InWidgetStyle, bParentEnabled);
-	LocalPaintTime += FPlatformTime::Seconds();
-	PaintTime = PaintTime * 0.95 + LocalPaintTime * 0.05;
-	return Result;
+	return SWorldPartitionEditorGrid::OnPaint(Args, AllottedGeometry, MyCullingRect, OutDrawElements, LayerId, InWidgetStyle, bParentEnabled);
 }
 
 int32 SWorldPartitionEditorGrid2D::PaintSoftwareCursor(const FGeometry& AllottedGeometry, const FSlateRect& MyCullingRect, FSlateWindowElementList& OutDrawElements, int32 LayerId) const
@@ -1313,7 +1352,19 @@ FReply SWorldPartitionEditorGrid2D::FocusSelection()
 	}
 	else
 	{
-		SelectionBox = WorldPartition->GetEditorWorldBounds();
+		// Override the minimap bounds if world partition minimap volumes exists
+		for (TActorIterator<AWorldPartitionMiniMapVolume> It(World); It; ++It)
+		{
+			if (AWorldPartitionMiniMapVolume* WorldPartitionMiniMapVolume = *It)
+			{
+				SelectionBox += WorldPartitionMiniMapVolume->GetBounds().GetBox();
+			}
+		}
+
+		if (!SelectionBox.IsValid)
+		{
+			SelectionBox = WorldPartition->GetEditorWorldBounds();
+		}
 	}
 
 	FocusBox(SelectionBox);
diff --git a/Engine/Source/Editor/WorldPartitionEditor/Public/WorldPartition/SWorldPartitionEditorGrid2D.h b/Engine/Source/Editor/WorldPartitionEditor/Public/WorldPartition/SWorldPartitionEditorGrid2D.h
--- a/Engine/Source/Editor/WorldPartitionEditor/Public/WorldPartition/SWorldPartitionEditorGrid2D.h
+++ b/Engine/Source/Editor/WorldPartitionEditor/Public/WorldPartition/SWorldPartitionEditorGrid2D.h
@@ -35,6 +35,7 @@ public:
 	SWorldPartitionEditorGrid2D();
 	~SWorldPartitionEditorGrid2D();
 
+protected:
 	void Construct(const FArguments& InArgs);
 
 	virtual void CreateRegionFromSelection();
@@ -60,7 +61,7 @@ public:
 	virtual int32 PaintSoftwareCursor(const FGeometry& AllottedGeometry, const FSlateRect& MyCullingRect, FSlateWindowElementList& OutDrawElements, int32 LayerId) const;
 
 	virtual FReply FocusSelection();
-protected:
+
 	void UpdateTransform() const;
 	void UpdateSelectionBox();
 	void ClearSelection();
@@ -91,7 +92,6 @@ protected:
 	bool bIsDragSelecting;
 	bool bIsPanning;
 	bool bShowActors;
-	bool bShowProfiling;
 	FVector2D MouseCursorPos;
 	FVector2D MouseCursorPosWorld;
 	FVector2D LastMouseCursorPosWorldDrag;
@@ -129,7 +129,6 @@ protected:
 	TSet<FGuid> ShownActorGuids;
 	TSet<FGuid> DirtyActorGuids;
 	FLoaderInterfaceSet ShownLoaderInterfaces;
-	FLoaderInterfaceSet HighlightedLoaderInterfaces;
 	FLoaderInterfaceSet HoveredLoaderInterfaces;
 	FLoaderInterfaceStack HoveredLoaderInterfacesStack;
 	FLoaderInterface HoveredLoaderInterface;
