World Partition Editor: allow snapping selection to effective grid size when holding Shift.

#rnx

diff --git a/Engine/Source/Editor/WorldPartitionEditor/Private/WorldPartition/SWorldPartitionEditorGrid2D.cpp b/Engine/Source/Editor/WorldPartitionEditor/Private/WorldPartition/SWorldPartitionEditorGrid2D.cpp
--- a/Engine/Source/Editor/WorldPartitionEditor/Private/WorldPartition/SWorldPartitionEditorGrid2D.cpp
+++ b/Engine/Source/Editor/WorldPartitionEditor/Private/WorldPartition/SWorldPartitionEditorGrid2D.cpp
@@ -171,6 +171,7 @@ SWorldPartitionEditorGrid2D::SWorldPartitionEditorGrid2D()
 	, bIsPanning(false)
 	, bShowActors(false)
 	, SelectBox(ForceInit)
+	, SelectBoxGridSnapped(ForceInit)
 	, TickTime(0)
 	, PaintTime(0)
 {
@@ -291,7 +292,7 @@ void SWorldPartitionEditorGrid2D::Construct(const FArguments& InArgs)
 
 	auto CanCreateRegionFromSelection = [this]()
 	{
-		return !!SelectBox.IsValid;
+		return !!SelectBoxGridSnapped.IsValid;
 	};
 
 	auto CanLoadUnloadSelectedRegions = [this](bool bLoad)
@@ -334,7 +335,7 @@ void SWorldPartitionEditorGrid2D::Construct(const FArguments& InArgs)
 
 void SWorldPartitionEditorGrid2D::CreateRegionFromSelection()
 {
-	const FBox RegionBox(FVector(SelectBox.Min.X, SelectBox.Min.Y, -HALF_WORLD_MAX), FVector(SelectBox.Max.X, SelectBox.Max.Y, HALF_WORLD_MAX));
+	const FBox RegionBox(FVector(SelectBoxGridSnapped.Min.X, SelectBoxGridSnapped.Min.Y, -HALF_WORLD_MAX), FVector(SelectBoxGridSnapped.Max.X, SelectBoxGridSnapped.Max.Y, HALF_WORLD_MAX));
 	UWorldPartitionEditorLoaderAdapter* EditorLoaderAdapter = WorldPartition->CreateEditorLoaderAdapter<FLoaderAdapterShape>(World, RegionBox, TEXT("Loaded Region"));
 	EditorLoaderAdapter->GetLoaderAdapter()->SetUserCreated(true);
 	EditorLoaderAdapter->GetLoaderAdapter()->Load();
@@ -357,6 +358,7 @@ void SWorldPartitionEditorGrid2D::LoadSelectedRegions()
 	}
 
 	SelectBox.Init();
+	SelectBoxGridSnapped.Init();
 		
 	GEditor->RedrawLevelEditingViewports();
 	Refresh();
@@ -380,6 +382,7 @@ void SWorldPartitionEditorGrid2D::UnloadSelectedRegions()
 	};
 		
 	SelectBox.Init();
+	SelectBoxGridSnapped.Init();
 
 	GEditor->RedrawLevelEditingViewports();
 	Refresh();
@@ -447,6 +450,11 @@ void SWorldPartitionEditorGrid2D::MoveCameraHere()
 	}
 }
 
+int32 SWorldPartitionEditorGrid2D::GetSelectionSnap() const
+{
+	return 1;
+}
+
 FReply SWorldPartitionEditorGrid2D::OnMouseButtonDown(const FGeometry& MyGeometry, const FPointerEvent& MouseEvent)
 {
 	const bool bIsLeftMouseButtonEffecting = MouseEvent.GetEffectingButton() == EKeys::LeftMouseButton;
@@ -473,6 +481,7 @@ FReply SWorldPartitionEditorGrid2D::OnMouseButtonDown(const FGeometry& MyGeometr
 			SelectionStart = MouseCursorPosWorld;
 			SelectionEnd = SelectionStart;
 			SelectBox.Init();
+			SelectBoxGridSnapped.Init();
 		}
 
 		return ReplyState;
@@ -534,7 +543,7 @@ FReply SWorldPartitionEditorGrid2D::OnMouseButtonUp(const FGeometry& MyGeometry,
 			FLoaderInterfaceSet LoaderAdaptersToSelect;
 			if (bIsDragSelecting)
 			{
-				ForEachIntersectingLoaderAdapters(WorldPartition, SelectBox, [&LoaderAdaptersToSelect](UObject* AdapterObject)
+				ForEachIntersectingLoaderAdapters(WorldPartition, SelectBoxGridSnapped, [&LoaderAdaptersToSelect](UObject* AdapterObject)
 				{
 					LoaderAdaptersToSelect.Add(AdapterObject);
 					return true;
@@ -561,6 +570,8 @@ FReply SWorldPartitionEditorGrid2D::OnMouseButtonUp(const FGeometry& MyGeometry,
 				SelectedLoaderInterfaces = MoveTemp(LoaderAdaptersToSelect);
 			}
 
+			SelectBox = SelectBoxGridSnapped;
+
 			bIsDragSelecting = false;
 		}
 
@@ -612,7 +623,7 @@ FReply SWorldPartitionEditorGrid2D::OnMouseMove(const FGeometry& MyGeometry, con
 			if (bIsDragSelecting)
 			{
 				SelectionEnd = MouseCursorPosWorld;
-				UpdateSelectionBox();
+				UpdateSelectionBox(MouseEvent.IsShiftDown());
 				return FReply::Handled();
 			}
 		}
@@ -861,7 +872,7 @@ uint32 SWorldPartitionEditorGrid2D::PaintActors(const FGeometry& AllottedGeometr
 			if (bShowBackground)
 			{
 				const FSlateColorBrush BackgroundBrush(FLinearColor::Black);
-				const FLinearColor LabelBackgroundColor(0, 0, 0, 0.25f);
+				const FLinearColor LabelBackgroundColor(0, 0, 0, 0.1f);
 
 				FSlateDrawElement::MakeBox(
 					OutDrawElements,
@@ -960,7 +971,7 @@ uint32 SWorldPartitionEditorGrid2D::PaintActors(const FGeometry& AllottedGeometr
 						// Outline
 						{
 							const bool bIsSelected = SelectedLoaderInterfaces.Contains(LoaderInterface);
-							const bool bIsInsideSelection = (SelectBox.GetVolume() > 0) && SelectBox.Intersect(AdapterBounds);
+							const bool bIsInsideSelection = (SelectBoxGridSnapped.GetVolume() > 0) && SelectBoxGridSnapped.Intersect(AdapterBounds);
 							const bool bIsLocalHovered = LocalHoveredLoaderAdapter == LoaderAdapter;
 							
 							float OutlineThickness = 2.0f;
@@ -1066,7 +1077,7 @@ uint32 SWorldPartitionEditorGrid2D::PaintActors(const FGeometry& AllottedGeometr
 
 					ActorColor = FLinearColor::Yellow;
 				}
-				else if ((SelectBox.GetVolume() > 0) && SelectBox.Intersect(ActorDescView.GetBounds()))
+				else if ((SelectBoxGridSnapped.GetVolume() > 0) && SelectBoxGridSnapped.Intersect(ActorDescView.GetBounds()))
 				{
 					ActorColor = FLinearColor::White;
 				}
@@ -1223,56 +1234,89 @@ uint32 SWorldPartitionEditorGrid2D::PaintViewer(const FGeometry& AllottedGeometr
 
 uint32 SWorldPartitionEditorGrid2D::PaintSelection(const FGeometry& AllottedGeometry, const FSlateRect& MyCullingRect, FSlateWindowElementList& OutDrawElements, uint32 LayerId) const
 {
-	if (SelectBox.IsValid)
+	if (SelectBoxGridSnapped.IsValid)
 	{
 		TArray<FVector2D> LinePoints;
 		LinePoints.SetNum(5);
 
-		FVector2D TopLeftW = FVector2D(SelectBox.Min);
-		FVector2D BottomRightW = FVector2D(SelectBox.Max);
-		FVector2D TopRightW = FVector2D(BottomRightW.X, TopLeftW.Y);
-		FVector2D BottomLeftW = FVector2D(TopLeftW.X, BottomRightW.Y);
+		// Draw snapped box
+		{
+			FVector2D TopLeftW = FVector2D(SelectBoxGridSnapped.Min);
+			FVector2D BottomRightW = FVector2D(SelectBoxGridSnapped.Max);
+			FVector2D TopRightW = FVector2D(BottomRightW.X, TopLeftW.Y);
+			FVector2D BottomLeftW = FVector2D(TopLeftW.X, BottomRightW.Y);
 
-		FVector2D TopLeft = WorldToScreen.TransformPoint(TopLeftW);
-		FVector2D BottomRight = WorldToScreen.TransformPoint(BottomRightW);
-		FVector2D TopRight = WorldToScreen.TransformPoint(TopRightW);
-		FVector2D BottomLeft = WorldToScreen.TransformPoint(BottomLeftW);
+			FVector2D TopLeft = WorldToScreen.TransformPoint(TopLeftW);
+			FVector2D BottomRight = WorldToScreen.TransformPoint(BottomRightW);
+			FVector2D TopRight = WorldToScreen.TransformPoint(TopRightW);
+			FVector2D BottomLeft = WorldToScreen.TransformPoint(BottomLeftW);
 
-		LinePoints[0] = TopLeft;
-		LinePoints[1] = TopRight;
-		LinePoints[2] = BottomRight;
-		LinePoints[3] = BottomLeft;
-		LinePoints[4] = TopLeft;
+			LinePoints[0] = TopLeft;
+			LinePoints[1] = TopRight;
+			LinePoints[2] = BottomRight;
+			LinePoints[3] = BottomLeft;
+			LinePoints[4] = TopLeft;
 
-		{
-			FSlateColorBrush CellBrush(FLinearColor::White);
-			FLinearColor CellColor(FLinearColor(1, 1, 1, 0.25f));
+			{
+				FSlateColorBrush CellBrush(FLinearColor::White);
+				FLinearColor CellColor(FLinearColor(1, 1, 1, 0.1f));
 
-			FPaintGeometry CellGeometry = AllottedGeometry.ToPaintGeometry(
-				TopLeft,
-				BottomRight - TopLeft
-			);
+				FPaintGeometry CellGeometry = AllottedGeometry.ToPaintGeometry(
+					TopLeft,
+					BottomRight - TopLeft
+				);
 
-			FSlateDrawElement::MakeBox(
-				OutDrawElements,
-				LayerId,
-				CellGeometry,
-				&CellBrush,
-				ESlateDrawEffect::None,
-				CellColor
+				FSlateDrawElement::MakeBox(
+					OutDrawElements,
+					LayerId,
+					CellGeometry,
+					&CellBrush,
+					ESlateDrawEffect::None,
+					CellColor
+				);
+			}
+
+			FSlateDrawElement::MakeLines(
+				OutDrawElements, 
+				LayerId, 
+				AllottedGeometry.ToPaintGeometry(), 
+				LinePoints, 
+				ESlateDrawEffect::None, 
+				FLinearColor::White, 
+				true, 
+				2.0f
 			);
 		}
 
-		FSlateDrawElement::MakeLines(
-			OutDrawElements, 
-			LayerId, 
-			AllottedGeometry.ToPaintGeometry(), 
-			LinePoints, 
-			ESlateDrawEffect::None, 
-			FLinearColor::White, 
-			true, 
-			2.0f
-		);
+		// Draw raw box
+		{
+			FVector2D TopLeftW = FVector2D(SelectBox.Min);
+			FVector2D BottomRightW = FVector2D(SelectBox.Max);
+			FVector2D TopRightW = FVector2D(BottomRightW.X, TopLeftW.Y);
+			FVector2D BottomLeftW = FVector2D(TopLeftW.X, BottomRightW.Y);
+
+			FVector2D TopLeft = WorldToScreen.TransformPoint(TopLeftW);
+			FVector2D BottomRight = WorldToScreen.TransformPoint(BottomRightW);
+			FVector2D TopRight = WorldToScreen.TransformPoint(TopRightW);
+			FVector2D BottomLeft = WorldToScreen.TransformPoint(BottomLeftW);
+
+			LinePoints[0] = TopLeft;
+			LinePoints[1] = TopRight;
+			LinePoints[2] = BottomRight;
+			LinePoints[3] = BottomLeft;
+			LinePoints[4] = TopLeft;
+
+			FSlateDrawElement::MakeLines(
+				OutDrawElements, 
+				LayerId, 
+				AllottedGeometry.ToPaintGeometry(), 
+				LinePoints, 
+				ESlateDrawEffect::None, 
+				FLinearColor::White, 
+				true, 
+				2.0f
+			);
+		}
 	}
 
 	return LayerId + 1;
@@ -1396,19 +1440,24 @@ void SWorldPartitionEditorGrid2D::UpdateTransform() const
 	ScreenToWorld = WorldToScreen.Inverse();
 }
 
-void SWorldPartitionEditorGrid2D::UpdateSelectionBox()
+void SWorldPartitionEditorGrid2D::UpdateSelectionBox(bool bSnap)
 {
-	const FBox2D SelectBox2D(
-		FVector2D::Min(SelectionStart, SelectionEnd), 
-		FVector2D::Max(SelectionStart, SelectionEnd)
-	);
+	const FBox2D SelectBox2D(FVector2D::Min(SelectionStart, SelectionEnd), FVector2D::Max(SelectionStart, SelectionEnd));
 
 	if (SelectBox2D.GetArea() > 0.0f)
 	{
-		SelectBox = FBox(
-			FVector(SelectBox2D.Min.X, SelectBox2D.Min.Y, -HALF_WORLD_MAX),
-			FVector(SelectBox2D.Max.X, SelectBox2D.Max.Y, HALF_WORLD_MAX)
-		);
+		const float MinX = SelectBox2D.Min.X;
+		const float MinY = SelectBox2D.Min.Y;
+		const float MaxX = SelectBox2D.Max.X;
+		const float MaxY = SelectBox2D.Max.Y;
+		SelectBox = FBox(FVector(MinX, MinY, -HALF_WORLD_MAX), FVector(MaxX, MaxY, HALF_WORLD_MAX));
+
+		const int32 SelectionSnap = bSnap ? GetSelectionSnap() : 1;
+		const float SnapMinX = FMath::GridSnap(MinX - SelectionSnap / 2, SelectionSnap);
+		const float SnapMinY = FMath::GridSnap(MinY - SelectionSnap / 2, SelectionSnap);
+		const float SnapMaxX = FMath::GridSnap(MaxX + SelectionSnap / 2, SelectionSnap);
+		const float SnapMaxY = FMath::GridSnap(MaxY + SelectionSnap / 2, SelectionSnap);
+		SelectBoxGridSnapped = FBox(FVector(SnapMinX, SnapMinY, -HALF_WORLD_MAX), FVector(SnapMaxX, SnapMaxY, HALF_WORLD_MAX));
 	}
 }
 
@@ -1416,6 +1465,7 @@ void SWorldPartitionEditorGrid2D::ClearSelection()
 {
 	SelectedLoaderInterfaces.Empty();
 	SelectBox.Init();
+	SelectBoxGridSnapped.Init();
 }
 
 #undef LOCTEXT_NAMESPACE
diff --git a/Engine/Source/Editor/WorldPartitionEditor/Private/WorldPartition/SWorldPartitionEditorGridSpatialHash.cpp b/Engine/Source/Editor/WorldPartitionEditor/Private/WorldPartition/SWorldPartitionEditorGridSpatialHash.cpp
--- a/Engine/Source/Editor/WorldPartitionEditor/Private/WorldPartition/SWorldPartitionEditorGridSpatialHash.cpp
+++ b/Engine/Source/Editor/WorldPartitionEditor/Private/WorldPartition/SWorldPartitionEditorGridSpatialHash.cpp
@@ -72,12 +72,11 @@ void SWorldPartitionEditorGridSpatialHash::UpdateWorldMiniMapDetails()
 	}
 }
 
-int32 SWorldPartitionEditorGridSpatialHash::PaintGrid(const FGeometry& AllottedGeometry, const FSlateRect& MyCullingRect, FSlateWindowElementList& OutDrawElements, int32 LayerId) const
+int32 SWorldPartitionEditorGridSpatialHash::GetSelectionSnap() const
 {
 	UWorldPartitionEditorSpatialHash* EditorSpatialHash = (UWorldPartitionEditorSpatialHash*)WorldPartition->EditorHash;
 	
-	// Found the best cell size depending on the current zoom
-	int64 EffectiveCellSize = EditorSpatialHash->CellSize;
+	int32 EffectiveCellSize = EditorSpatialHash->CellSize;
 	for (const UWorldPartitionEditorSpatialHash::FCellNodeHashLevel& HashLevel : EditorSpatialHash->HashLevels)
 	{
 		const FVector2D CellScreenSize = WorldToScreen.TransformVector(FVector2D(EffectiveCellSize, EffectiveCellSize));
@@ -88,6 +87,16 @@ int32 SWorldPartitionEditorGridSpatialHash::PaintGrid(const FGeometry& AllottedG
 
 		EffectiveCellSize *= 2;
 	}
+
+	return EffectiveCellSize;
+}
+
+int32 SWorldPartitionEditorGridSpatialHash::PaintGrid(const FGeometry& AllottedGeometry, const FSlateRect& MyCullingRect, FSlateWindowElementList& OutDrawElements, int32 LayerId) const
+{
+	UWorldPartitionEditorSpatialHash* EditorSpatialHash = (UWorldPartitionEditorSpatialHash*)WorldPartition->EditorHash;
+	
+	// Found the best cell size depending on the current zoom
+	const int64 EffectiveCellSize = GetSelectionSnap();
 	
 	// Compute visible rect
 	const FBox2D ViewRect(FVector2D(ForceInitToZero), AllottedGeometry.GetLocalSize());
diff --git a/Engine/Source/Editor/WorldPartitionEditor/Public/WorldPartition/SWorldPartitionEditorGrid2D.h b/Engine/Source/Editor/WorldPartitionEditor/Public/WorldPartition/SWorldPartitionEditorGrid2D.h
--- a/Engine/Source/Editor/WorldPartitionEditor/Public/WorldPartition/SWorldPartitionEditorGrid2D.h
+++ b/Engine/Source/Editor/WorldPartitionEditor/Public/WorldPartition/SWorldPartitionEditorGrid2D.h
@@ -38,12 +38,14 @@ public:
 protected:
 	void Construct(const FArguments& InArgs);
 
-	virtual void CreateRegionFromSelection();
-	virtual void LoadSelectedRegions();
-	virtual void UnloadSelectedRegions();
-	virtual void ConvertSelectedRegionsToActors();
+	void CreateRegionFromSelection();
+	void LoadSelectedRegions();
+	void UnloadSelectedRegions();
+	void ConvertSelectedRegionsToActors();
 	void MoveCameraHere();
 
+	virtual int32 GetSelectionSnap() const;
+
 	virtual void Tick(const FGeometry& AllottedGeometry, const double InCurrentTime, const float InDeltaTime) override;
 	virtual FReply OnMouseButtonDown(const FGeometry& MyGeometry, const FPointerEvent& MouseEvent) override;
 	virtual FReply OnMouseButtonUp(const FGeometry& MyGeometry, const FPointerEvent& MouseEvent) override;
@@ -63,7 +65,7 @@ protected:
 	virtual FReply FocusSelection();
 
 	void UpdateTransform() const;
-	void UpdateSelectionBox();
+	void UpdateSelectionBox(bool bSnap);
 	void ClearSelection();
 
 	const TSharedRef<FUICommandList> CommandList;
@@ -98,6 +100,7 @@ protected:
 	FVector2D SelectionStart;
 	FVector2D SelectionEnd;
 	FBox SelectBox;
+	FBox SelectBoxGridSnapped;
 	FSlateFontInfo SmallLayoutFont;
 	float TotalMouseDelta;
 
diff --git a/Engine/Source/Editor/WorldPartitionEditor/Public/WorldPartition/SWorldPartitionEditorGridSpatialHash.h b/Engine/Source/Editor/WorldPartitionEditor/Public/WorldPartition/SWorldPartitionEditorGridSpatialHash.h
--- a/Engine/Source/Editor/WorldPartitionEditor/Public/WorldPartition/SWorldPartitionEditorGridSpatialHash.h
+++ b/Engine/Source/Editor/WorldPartitionEditor/Public/WorldPartition/SWorldPartitionEditorGridSpatialHash.h
@@ -16,7 +16,8 @@ public:
 
 	void Construct(const FArguments& InArgs);
 
-	int32 PaintGrid(const FGeometry& AllottedGeometry, const FSlateRect& MyCullingRect, FSlateWindowElementList& OutDrawElements, int32 LayerId) const override;
+	virtual int32 GetSelectionSnap() const override;
+	virtual int32 PaintGrid(const FGeometry& AllottedGeometry, const FSlateRect& MyCullingRect, FSlateWindowElementList& OutDrawElements, int32 LayerId) const override;
 
 private:
 	void UpdateWorldMiniMapDetails();
