World Partition: remove support for completely disabling streaming for partitioned maps.

#rnx

diff --git a/Engine/Source/Editor/WorldPartitionEditor/Private/WorldPartition/Customizations/WorldPartitionDetailsCustomization.cpp b/Engine/Source/Editor/WorldPartitionEditor/Private/WorldPartition/Customizations/WorldPartitionDetailsCustomization.cpp
--- a/Engine/Source/Editor/WorldPartitionEditor/Private/WorldPartition/Customizations/WorldPartitionDetailsCustomization.cpp
+++ b/Engine/Source/Editor/WorldPartitionEditor/Private/WorldPartition/Customizations/WorldPartitionDetailsCustomization.cpp
@@ -33,10 +33,11 @@ void FWorldPartitionDetails::CustomizeDetails(IDetailLayoutBuilder& InDetailBuil
 	WorldPartition = CastChecked<UWorldPartition>(ObjectsBeingCustomized[0].Get());
 	RuntimeHashClass = (WorldPartition.IsValid() && WorldPartition->RuntimeHash) ? WorldPartition->RuntimeHash->GetClass() : nullptr;
 
-	IDetailCategoryBuilder& WorldPartitionCategory = InDetailBuilder.EditCategory("WorldPartition");
+	IDetailCategoryBuilder& WorldPartitionCategory = InDetailBuilder.EditCategory("WorldPartitionSetup");
 
 	// Enable streaming checkbox
 	WorldPartitionCategory.AddCustomRow(LOCTEXT("EnableStreaming", "Enable Streaming"), false)
+		.RowTag(TEXT("EnableStreaming"))
 		.NameContent()
 		[
 			SNew(STextBlock)
@@ -54,6 +55,7 @@ void FWorldPartitionDetails::CustomizeDetails(IDetailLayoutBuilder& InDetailBuil
 
 	// Runtime hash class selector
 	WorldPartitionCategory.AddCustomRow(LOCTEXT("RuntimeHashClass", "Runtime Hash Class"), false)
+		.RowTag(TEXT("RuntimeHashClass"))
 		.NameContent()
 		[
 			SNew(STextBlock)
diff --git a/Engine/Source/Runtime/Engine/Private/WorldPartition/WorldPartition.cpp b/Engine/Source/Runtime/Engine/Private/WorldPartition/WorldPartition.cpp
--- a/Engine/Source/Runtime/Engine/Private/WorldPartition/WorldPartition.cpp
+++ b/Engine/Source/Runtime/Engine/Private/WorldPartition/WorldPartition.cpp
@@ -216,7 +216,6 @@ UWorldPartition::UWorldPartition(const FObjectInitializer& ObjectInitializer)
 	, bStreamingWasEnabled(true)
 	, bShouldEnableStreamingWarned(false)
 	, bShouldCheckEnableStreamingWarning(false)
-	, bSupportsStreaming(true)
 	, bCanBeUsedByLevelInstance(false)
 	, bForceGarbageCollection(false)
 	, bForceGarbageCollectionPurge(false)
@@ -928,7 +927,7 @@ const FTransform& UWorldPartition::GetInstanceTransform() const
 #if WITH_EDITOR
 bool UWorldPartition::SupportsStreaming() const
 {
-	return bSupportsStreaming && (World ? World->GetWorldSettings()->SupportsWorldPartitionStreaming() : false);
+	return World ? World->GetWorldSettings()->SupportsWorldPartitionStreaming() : false;
 }
 
 bool UWorldPartition::IsStreamingEnabled() const
diff --git a/Engine/Source/Runtime/Engine/Private/WorldPartition/WorldPartitionConverter.cpp b/Engine/Source/Runtime/Engine/Private/WorldPartition/WorldPartitionConverter.cpp
--- a/Engine/Source/Runtime/Engine/Private/WorldPartition/WorldPartitionConverter.cpp
+++ b/Engine/Source/Runtime/Engine/Private/WorldPartition/WorldPartitionConverter.cpp
@@ -60,7 +60,6 @@ bool FWorldPartitionConverter::Convert()
 		FLevelActorFoldersHelper::SetUseActorFolders(World->PersistentLevel, Parameters.bUseActorFolders);
 		WorldPartition->bEnableStreaming = Parameters.bEnableStreaming;
 		WorldPartition->bStreamingWasEnabled = Parameters.bEnableStreaming;
-		WorldPartition->bSupportsStreaming = Parameters.bSupportsStreaming;
 		WorldPartition->SetCanBeUsedByLevelInstance(!Parameters.bEnableStreaming);
 		bCreatedWorldPartition = true;
 	}
diff --git a/Engine/Source/Runtime/Engine/Public/WorldPartition/WorldPartition.h b/Engine/Source/Runtime/Engine/Public/WorldPartition/WorldPartition.h
--- a/Engine/Source/Runtime/Engine/Public/WorldPartition/WorldPartition.h
+++ b/Engine/Source/Runtime/Engine/Public/WorldPartition/WorldPartition.h
@@ -256,10 +256,6 @@ private:
 	/** Used to know if we need to recheck if the user should enable streaming based on world size. */
 	bool bShouldCheckEnableStreamingWarning;
 
-	/** Used to check if the current world supports enabling streaming. */
-	UPROPERTY()
-	bool bSupportsStreaming;
-
 	/** Whether Level Instance can reference this partition. */
 	UPROPERTY(EditAnywhere, Category = "WorldPartition", AdvancedDisplay, meta = (EditConditionHides, HideEditConditionToggle, EditCondition = "!bEnableStreaming"))
 	bool bCanBeUsedByLevelInstance;
diff --git a/Engine/Source/Runtime/Engine/Public/WorldPartition/WorldPartitionConverter.h b/Engine/Source/Runtime/Engine/Public/WorldPartition/WorldPartitionConverter.h
--- a/Engine/Source/Runtime/Engine/Public/WorldPartition/WorldPartitionConverter.h
+++ b/Engine/Source/Runtime/Engine/Public/WorldPartition/WorldPartitionConverter.h
@@ -21,7 +21,6 @@ public:
 		FParameters()
 		: bConvertSubLevels(true)
 		, bEnableStreaming(true)
-		, bSupportsStreaming(true)
 		, bUseActorFolders(true)
 		, EditorHashClass(nullptr)
 		, RuntimeHashClass(nullptr)
@@ -29,7 +28,6 @@ public:
 
 		bool bConvertSubLevels;
 		bool bEnableStreaming;
-		bool bSupportsStreaming;
 		bool bUseActorFolders;
 		TSubclassOf<UWorldPartitionEditorHash> EditorHashClass;
 		TSubclassOf<UWorldPartitionRuntimeHash> RuntimeHashClass;
