World Partition Editor: add a shadow to the camera indicato so it shows on white backgrounds.

#rnx

diff --git a/Engine/Source/Editor/WorldPartitionEditor/Private/WorldPartition/SWorldPartitionEditorGrid2D.cpp b/Engine/Source/Editor/WorldPartitionEditor/Private/WorldPartition/SWorldPartitionEditorGrid2D.cpp
--- a/Engine/Source/Editor/WorldPartitionEditor/Private/WorldPartition/SWorldPartitionEditorGrid2D.cpp
+++ b/Engine/Source/Editor/WorldPartitionEditor/Private/WorldPartition/SWorldPartitionEditorGrid2D.cpp
@@ -1180,53 +1180,59 @@ uint32 SWorldPartitionEditorGrid2D::PaintTextInfo(const FGeometry& AllottedGeome
 
 uint32 SWorldPartitionEditorGrid2D::PaintViewer(const FGeometry& AllottedGeometry, const FSlateRect& MyCullingRect, FSlateWindowElementList& OutDrawElements, uint32 LayerId) const
 {
-	FVector ObserverPosition;
-	FRotator ObserverRotation;
-	if (GetObserverView(ObserverPosition, ObserverRotation))
+	auto MakeRotatedBoxWithShadow = [this, &AllottedGeometry, &OutDrawElements, &LayerId](const FVector2D& Location, const FRotator& Rotation, const FSlateBrush* Image, const FLinearColor& Color, const FVector2D& ShadowSize)
 	{
-		FVector2D LocalViewLocation = WorldToScreen.TransformPoint(FVector2D(ObserverPosition));
-		const FSlateBrush* CameraImage = FAppStyle::GetBrush(TEXT("WorldPartition.SimulationViewPosition"));
-	
-		FPaintGeometry PaintGeometry = AllottedGeometry.ToPaintGeometry(
-			LocalViewLocation - CameraImage->ImageSize * 0.5f, 
-			CameraImage->ImageSize
+		const FVector2D LocalLocation = WorldToScreen.TransformPoint(Location);
+		const FPaintGeometry PaintGeometryShadow = AllottedGeometry.ToPaintGeometry(
+			LocalLocation - (Image->ImageSize + ShadowSize) * 0.5f, 
+			Image->ImageSize + ShadowSize
 		);
 
 		FSlateDrawElement::MakeRotatedBox(
 			OutDrawElements,
 			++LayerId,
-			PaintGeometry,
-			CameraImage,
+			PaintGeometryShadow,
+			Image,
 			ESlateDrawEffect::None,
-			FMath::DegreesToRadians(ObserverRotation.Yaw),
-			CameraImage->ImageSize * 0.5f,
-			FSlateDrawElement::RelativeToElement
+			FMath::DegreesToRadians(Rotation.Yaw),
+			(Image->ImageSize + ShadowSize) * 0.5f,
+			FSlateDrawElement::RelativeToElement,
+			FLinearColor::Black
 		);
-	}
 
-	FVector PlayerPosition;
-	FRotator PlayerRotation;
-	if (GetPlayerView(PlayerPosition, PlayerRotation))
-	{
-		FVector2D LocalViewLocation = WorldToScreen.TransformPoint(FVector2D(PlayerPosition));
-		const FSlateBrush* CameraImage = FAppStyle::GetBrush(TEXT("WorldPartition.SimulationViewPosition"));
-	
-		FPaintGeometry PaintGeometry = AllottedGeometry.ToPaintGeometry(
-			LocalViewLocation - CameraImage->ImageSize * 0.5f, 
-			CameraImage->ImageSize
+		const FPaintGeometry PaintGeometry = AllottedGeometry.ToPaintGeometry(
+			LocalLocation - Image->ImageSize * 0.5f, 
+			Image->ImageSize
 		);
 
 		FSlateDrawElement::MakeRotatedBox(
 			OutDrawElements,
 			++LayerId,
 			PaintGeometry,
-			CameraImage,
+			Image,
 			ESlateDrawEffect::None,
-			FMath::DegreesToRadians(PlayerRotation.Yaw),
-			CameraImage->ImageSize * 0.5f,
+			FMath::DegreesToRadians(Rotation.Yaw),
+			Image->ImageSize * 0.5f,
 			FSlateDrawElement::RelativeToElement,
-			FLinearColor(FColorList::Orange)
+			Color
 		);
+	};
+
+	const FVector2D ShadowSize(2, 2);
+	const FSlateBrush* CameraImage = FAppStyle::GetBrush(TEXT("WorldPartition.SimulationViewPosition"));
+
+	FVector ObserverPosition;
+	FRotator ObserverRotation;
+	if (GetObserverView(ObserverPosition, ObserverRotation))
+	{
+		MakeRotatedBoxWithShadow(FVector2D(ObserverPosition), ObserverRotation, CameraImage, FLinearColor::White, ShadowSize);
+	}
+
+	FVector PlayerPosition;
+	FRotator PlayerRotation;
+	if (GetPlayerView(PlayerPosition, PlayerRotation))
+	{
+		MakeRotatedBoxWithShadow(FVector2D(PlayerPosition), PlayerRotation, CameraImage, FColorList::Orange, ShadowSize);
 	}
 
 	return LayerId + 1;
