World Partition Editor: added option to follow player in PIE.

#rnx

diff --git a/Engine/Source/Editor/WorldPartitionEditor/Private/WorldPartition/SWorldPartitionEditorGrid2D.cpp b/Engine/Source/Editor/WorldPartitionEditor/Private/WorldPartition/SWorldPartitionEditorGrid2D.cpp
--- a/Engine/Source/Editor/WorldPartitionEditor/Private/WorldPartition/SWorldPartitionEditorGrid2D.cpp
+++ b/Engine/Source/Editor/WorldPartitionEditor/Private/WorldPartition/SWorldPartitionEditorGrid2D.cpp
@@ -170,6 +170,7 @@ SWorldPartitionEditorGrid2D::SWorldPartitionEditorGrid2D()
 	, bIsDragSelecting(false)
 	, bIsPanning(false)
 	, bShowActors(false)
+	, bFollowPlayerInPIE(false)
 	, SelectBox(ForceInit)
 	, SelectBoxGridSnapped(ForceInit)
 	, TickTime(0)
@@ -231,6 +232,23 @@ void SWorldPartitionEditorGrid2D::Construct(const FArguments& InArgs)
 					]
 					+SHorizontalBox::Slot()
 					.AutoWidth()
+					[
+						SNew(SCheckBox)
+						.IsChecked(bFollowPlayerInPIE ? ECheckBoxState::Checked : ECheckBoxState::Unchecked)
+						.IsEnabled(true)
+						.OnCheckStateChanged(FOnCheckStateChanged::CreateLambda([=](ECheckBoxState State) { bFollowPlayerInPIE = !bFollowPlayerInPIE; }))
+					]
+					+SHorizontalBox::Slot()
+					.FillWidth(1.0f)
+					.VAlign(VAlign_Center)
+					[
+						SNew(STextBlock)
+						.AutoWrapText(true)
+						.IsEnabled(true)
+						.Text(LOCTEXT("FollowPlayerInPIE", "Follow Player in PIE"))
+					]
+					+SHorizontalBox::Slot()
+					.AutoWidth()
 					[
 						SNew(SCheckBox)
 						.IsChecked(GetMutableDefault<UWorldPartitionEditorPerProjectUserSettings>()->GetBugItGoLoadRegion() ? ECheckBoxState::Checked : ECheckBoxState::Unchecked)
@@ -269,6 +287,7 @@ void SWorldPartitionEditorGrid2D::Construct(const FArguments& InArgs)
 						SNew(SButton)
 						.Text(LOCTEXT("FocusSelection", "Focus Selection"))
 						.OnClicked(this, &SWorldPartitionEditorGrid2D::FocusSelection)
+						.IsEnabled_Lambda([this]() { return IsInteractive(); })
 					]
 				]
 			]
@@ -450,6 +469,16 @@ void SWorldPartitionEditorGrid2D::MoveCameraHere()
 	}
 }
 
+bool SWorldPartitionEditorGrid2D::IsFollowPlayerInPIE() const
+{
+	return bFollowPlayerInPIE && (UWorldPartition::IsSimulating() || GEditor->PlayWorld);
+}
+
+bool SWorldPartitionEditorGrid2D::IsInteractive() const
+{
+	return !IsFollowPlayerInPIE();
+}
+
 int32 SWorldPartitionEditorGrid2D::GetSelectionSnap() const
 {
 	return 1;
@@ -593,7 +622,10 @@ FReply SWorldPartitionEditorGrid2D::OnMouseButtonUp(const FGeometry& MyGeometry,
 
 FReply SWorldPartitionEditorGrid2D::OnMouseButtonDoubleClick(const FGeometry& InMyGeometry, const FPointerEvent& InMouseEvent)
 {
-	MoveCameraHere();
+	if (IsInteractive())
+	{
+		MoveCameraHere();
+	}
 	return FReply::Handled();
 }
 
@@ -611,7 +643,7 @@ FReply SWorldPartitionEditorGrid2D::OnMouseMove(const FGeometry& MyGeometry, con
 		const bool bIsRightMouseButtonDown = MouseEvent.IsMouseButtonDown(EKeys::RightMouseButton);
 		const bool bIsLeftMouseButtonDown = MouseEvent.IsMouseButtonDown(EKeys::LeftMouseButton);
 		const bool bIsMiddleMouseButtonDown = MouseEvent.IsMouseButtonDown(EKeys::MiddleMouseButton);
-		const bool bIsDragTrigger = TotalMouseDelta > FSlateApplication::Get().GetDragTriggerDistance();
+		const bool bIsDragTrigger = IsInteractive() && (TotalMouseDelta > FSlateApplication::Get().GetDragTriggerDistance());
 
 		if (bIsLeftMouseButtonDown)
 		{
@@ -650,13 +682,16 @@ FReply SWorldPartitionEditorGrid2D::OnMouseMove(const FGeometry& MyGeometry, con
 
 FReply SWorldPartitionEditorGrid2D::OnMouseWheel(const FGeometry& MyGeometry, const FPointerEvent& MouseEvent)
 {
-	FVector2D MousePosLocalSpace = MouseCursorPos - MyGeometry.GetLocalSize() * 0.5f;
-	FVector2D P0 = MousePosLocalSpace / Scale;
-	float Delta = 1.0f + FMath::Abs(MouseEvent.GetWheelDelta() / 4.0f);
-	Scale = FMath::Clamp(Scale * (MouseEvent.GetWheelDelta() > 0 ? Delta : (1.0f / Delta)), 0.00000001f, 10.0f);
-	FVector2D P1 = MousePosLocalSpace / Scale;
-	Trans += (P1 - P0);
-	UpdateTransform();
+	if (IsInteractive())
+	{
+		FVector2D MousePosLocalSpace = MouseCursorPos - MyGeometry.GetLocalSize() * 0.5f;
+		FVector2D P0 = MousePosLocalSpace / Scale;
+		float Delta = 1.0f + FMath::Abs(MouseEvent.GetWheelDelta() / 4.0f);
+		Scale = FMath::Clamp(Scale * (MouseEvent.GetWheelDelta() > 0 ? Delta : (1.0f / Delta)), 0.00000001f, 10.0f);
+		FVector2D P1 = MousePosLocalSpace / Scale;
+		Trans += (P1 - P0);
+		UpdateTransform();
+	}
 	return FReply::Handled();
 }
 
@@ -1443,8 +1478,27 @@ void SWorldPartitionEditorGrid2D::FocusBox(const FBox& Box) const
 
 void SWorldPartitionEditorGrid2D::UpdateTransform() const
 {
-	FTransform2d T(1.0f, Trans);
-	FTransform2d V(Scale, FVector2D(ScreenRect.GetSize().X * 0.5f, ScreenRect.GetSize().Y * 0.5f));
+	FTransform2d T;
+	FTransform2d V;
+
+	if (IsFollowPlayerInPIE())
+	{
+		FVector FollowPosition;
+		FRotator FollowRotation;
+
+		if (UWorldPartition::IsSimulating() ? GetObserverView(FollowPosition, FollowRotation) : GetPlayerView(FollowPosition, FollowRotation))
+		{
+			const FVector2D PlayerExtent = FVector2D(10000);
+			T = FTransform2d(1.0f, -FVector2D(FollowPosition));
+			V = FTransform2d((ScreenRect.GetExtent() / PlayerExtent).GetMin(), FVector2D(ScreenRect.GetSize().X * 0.5f, ScreenRect.GetSize().Y * 0.5f));
+		}
+	}
+	else
+	{
+		T = FTransform2d(1.0f, Trans);
+		V = FTransform2d(Scale, FVector2D(ScreenRect.GetSize().X * 0.5f, ScreenRect.GetSize().Y * 0.5f));
+	}
+
 	WorldToScreen = T.Concatenate(V);
 	ScreenToWorld = WorldToScreen.Inverse();
 }
diff --git a/Engine/Source/Editor/WorldPartitionEditor/Public/WorldPartition/SWorldPartitionEditorGrid2D.h b/Engine/Source/Editor/WorldPartitionEditor/Public/WorldPartition/SWorldPartitionEditorGrid2D.h
--- a/Engine/Source/Editor/WorldPartitionEditor/Public/WorldPartition/SWorldPartitionEditorGrid2D.h
+++ b/Engine/Source/Editor/WorldPartitionEditor/Public/WorldPartition/SWorldPartitionEditorGrid2D.h
@@ -44,6 +44,9 @@ protected:
 	void ConvertSelectedRegionsToActors();
 	void MoveCameraHere();
 
+	bool IsFollowPlayerInPIE() const;
+	bool IsInteractive() const;
+
 	virtual int32 GetSelectionSnap() const;
 
 	virtual void Tick(const FGeometry& AllottedGeometry, const double InCurrentTime, const float InDeltaTime) override;
@@ -94,6 +97,7 @@ protected:
 	bool bIsDragSelecting;
 	bool bIsPanning;
 	bool bShowActors;
+	bool bFollowPlayerInPIE;
 	FVector2D MouseCursorPos;
 	FVector2D MouseCursorPosWorld;
 	FVector2D LastMouseCursorPosWorldDrag;
