World Partition Editor Improvements:
- Moved all minimap related code from the editor spatial hash specific implementation to base class.
- Better minimap shadowing by clipping loading regions from the shadowed area.
- Added drop shadow to all actor names rendering.
- Added new button to focus on loaded regions.

#rnx

diff --git a/Engine/Source/Editor/WorldPartitionEditor/Private/WorldPartition/SWorldPartitionEditorGrid.cpp b/Engine/Source/Editor/WorldPartitionEditor/Private/WorldPartition/SWorldPartitionEditorGrid.cpp
--- a/Engine/Source/Editor/WorldPartitionEditor/Private/WorldPartition/SWorldPartitionEditorGrid.cpp
+++ b/Engine/Source/Editor/WorldPartitionEditor/Private/WorldPartition/SWorldPartitionEditorGrid.cpp
@@ -16,6 +16,7 @@ TMap<FName, SWorldPartitionEditorGrid::PartitionEditorGridCreateInstanceFunc> SW
 void SWorldPartitionEditorGrid::Construct(const FArguments& InArgs)
 {
 	World = InArgs._InWorld;
+	WorldPartition = World ? World->GetWorldPartition() : nullptr;
 
 	if (IsDisabled())
 	{
diff --git a/Engine/Source/Editor/WorldPartitionEditor/Private/WorldPartition/SWorldPartitionEditorGrid2D.cpp b/Engine/Source/Editor/WorldPartitionEditor/Private/WorldPartition/SWorldPartitionEditorGrid2D.cpp
--- a/Engine/Source/Editor/WorldPartitionEditor/Private/WorldPartition/SWorldPartitionEditorGrid2D.cpp
+++ b/Engine/Source/Editor/WorldPartitionEditor/Private/WorldPartition/SWorldPartitionEditorGrid2D.cpp
@@ -3,8 +3,8 @@
 #include "WorldPartition/SWorldPartitionEditorGrid2D.h"
 #include "Brushes/SlateColorBrush.h"
 #include "Engine/Selection.h"
+#include "EngineModule.h"
 #include "EngineUtils.h"
-#include "GameFramework/WorldSettings.h"
 #include "Modules/ModuleManager.h"
 #include "WorldPartition/WorldPartition.h"
 #include "WorldPartition/WorldPartitionEditorLoaderAdapter.h"
@@ -13,20 +13,21 @@
 #include "WorldPartition/WorldPartitionActorDescView.h"
 #include "WorldPartition/WorldPartitionEditorPerProjectUserSettings.h"
 #include "WorldPartition/WorldPartitionHelpers.h"
+#include "WorldPartition/WorldPartitionMiniMap.h"
+#include "WorldPartition/WorldPartitionMiniMapHelper.h"
 #include "WorldPartition/WorldPartitionMiniMapVolume.h"
 #include "WorldPartition/LoaderAdapter/LoaderAdapterShape.h"
 #include "ActorFactories/ActorFactory.h"
 #include "Builders/CubeBuilder.h"
-#include "Editor/GroupActor.h"
 #include "Editor/EditorEngine.h"
 #include "Fonts/FontMeasure.h"
 #include "Widgets/Input/SCheckBox.h"
 #include "Widgets/Input/SButton.h"
 #include "Rendering/SlateRenderer.h"
+#include "Renderer/Private/RendererModule.h"
 #include "WorldBrowserModule.h"
 #include "LevelEditorViewport.h"
 #include "LocationVolume.h"
-#include "Algo/Transform.h"
 #include "SWorldPartitionViewportWidget.h"
 
 #define LOCTEXT_NAMESPACE "WorldPartitionEditor"
@@ -57,6 +58,77 @@ private:
 	double& Value;
 };
 
+class F2DRectBooleanSubtract
+{
+public:
+	F2DRectBooleanSubtract(const FBox2D& Rect)
+	{
+		Rects.Add(Rect);
+	};
+
+	void SubRect(const FBox2D& SubRect)
+	{
+		TArray<FBox2D> NewRects;
+
+		for (auto It(Rects.CreateIterator()); It; ++It)
+		{
+			if (SubRect.Intersect(*It))
+			{
+				if (SubRect.IsInsideOrOn(It->Min) && SubRect.IsInsideOrOn(It->Max))
+				{
+					It.RemoveCurrent();
+					continue;
+				}
+
+				const FBox2D ClipRect = It->Overlap(SubRect);
+
+				if (ClipRect.GetArea())
+				{
+					const bool bClipTop = (ClipRect.Min.Y > It->Min.Y) && (ClipRect.Min.Y < It->Max.Y);
+					const bool bClipBottom = (ClipRect.Max.Y < It->Max.Y) && (ClipRect.Max.Y > It->Min.Y);
+					const bool bClipLeft = (ClipRect.Min.X > It->Min.X) && (ClipRect.Min.X < It->Max.X);
+					const bool bClipRight = (ClipRect.Max.X < It->Max.X) && (ClipRect.Max.X > It->Min.X);
+
+					if (bClipTop || bClipBottom || bClipLeft || bClipRight)
+					{
+						if (bClipTop)
+						{
+							NewRects.Emplace(It->Min, FVector2D(It->Max.X, ClipRect.Min.Y));
+						}
+
+						if (bClipBottom)
+						{
+							NewRects.Emplace(FVector2D(It->Min.X, ClipRect.Max.Y), It->Max);
+						}
+
+						if (bClipLeft)
+						{
+							NewRects.Emplace(FVector2D(It->Min.X, ClipRect.Min.Y), FVector2D(ClipRect.Min.X, ClipRect.Max.Y));
+						}
+
+						if (bClipRight)
+						{
+ 							NewRects.Emplace(FVector2D(ClipRect.Max.X, ClipRect.Min.Y), FVector2D(It->Max.X, ClipRect.Max.Y));
+						}
+
+						It.RemoveCurrent();
+					}
+				}
+			}
+		}
+
+		Rects += NewRects;
+	}
+
+	const TArray<FBox2D> GetRects() const
+	{
+		return Rects;
+	}
+
+private:
+	TArray<FBox2D> Rects;
+};
+
 static bool IsBoundsSelected(const FBox& SelectBox, const FBox& Bounds)
 {
 	return SelectBox.IsValid && Bounds.IntersectXY(SelectBox) && !Bounds.IsInsideXY(SelectBox);
@@ -173,6 +245,7 @@ SWorldPartitionEditorGrid2D::SWorldPartitionEditorGrid2D()
 	, bFollowPlayerInPIE(false)
 	, SelectBox(ForceInit)
 	, SelectBoxGridSnapped(ForceInit)
+	, WorldMiniMapBounds(ForceInit)
 	, TickTime(0)
 	, PaintTime(0)
 {
@@ -189,6 +262,12 @@ void SWorldPartitionEditorGrid2D::Construct(const FArguments& InArgs)
 {
 	SWorldPartitionEditorGrid::Construct(SWorldPartitionEditorGrid::FArguments().InWorld(InArgs._InWorld));
 
+	if (WorldPartition)
+	{
+		UpdateWorldMiniMapDetails();
+		bShowActors = !WorldMiniMapBrush.HasUObject();
+	}
+
 	// Defaults
 	Trans = FVector2D(0, 0);
 	Scale = 0.00133333332;
@@ -285,10 +364,18 @@ void SWorldPartitionEditorGrid2D::Construct(const FArguments& InArgs)
 					.AutoWidth()
 					[
 						SNew(SButton)
-						.Text(LOCTEXT("FocusSelection", "Focus Selection"))
+						.Text_Lambda([this]() { return GEditor->GetSelectedActors()->Num() ? LOCTEXT("FocusSelection", "Focus Selection") : LOCTEXT("FocusWorld", "Focus World"); })
 						.OnClicked(this, &SWorldPartitionEditorGrid2D::FocusSelection)
 						.IsEnabled_Lambda([this]() { return IsInteractive(); })
 					]
+					+ SHorizontalBox::Slot()
+					.AutoWidth()
+					[
+						SNew(SButton)
+						.Text(LOCTEXT("FocusLoadedRegions", "Focus Loaded Regions"))
+						.OnClicked(this, &SWorldPartitionEditorGrid2D::FocusLoadedRegions)
+						.IsEnabled_Lambda([this]() { return IsInteractive() && WorldPartition && WorldPartition->HasLoadedUserCreatedRegions(); })
+					]
 				]
 			]
 		]
@@ -352,6 +439,20 @@ void SWorldPartitionEditorGrid2D::Construct(const FArguments& InArgs)
 	ActionList.MapAction(Commands.MoveCameraHere, FExecuteAction::CreateSP(this, &SWorldPartitionEditorGrid2D::MoveCameraHere));
 }
 
+void SWorldPartitionEditorGrid2D::UpdateWorldMiniMapDetails()
+{
+	if (AWorldPartitionMiniMap* WorldMiniMap = FWorldPartitionMiniMapHelper::GetWorldPartitionMiniMap(World))
+	{
+		WorldMiniMapBounds = FBox2D(FVector2D(WorldMiniMap->MiniMapWorldBounds.Min), FVector2D(WorldMiniMap->MiniMapWorldBounds.Max));
+		if (UTexture2D* MiniMapTexture = WorldMiniMap->MiniMapTexture)
+		{
+			WorldMiniMapBrush.SetUVRegion(WorldMiniMap->UVOffset);
+			WorldMiniMapBrush.SetImageSize(MiniMapTexture->GetImportedSize());
+			WorldMiniMapBrush.SetResourceObject(MiniMapTexture);
+		}
+	}
+}
+
 void SWorldPartitionEditorGrid2D::CreateRegionFromSelection()
 {
 	const FBox RegionBox(FVector(SelectBoxGridSnapped.Min.X, SelectBoxGridSnapped.Min.Y, -HALF_WORLD_MAX), FVector(SelectBoxGridSnapped.Max.X, SelectBoxGridSnapped.Max.Y, HALF_WORLD_MAX));
@@ -927,6 +1028,17 @@ uint32 SWorldPartitionEditorGrid2D::PaintActors(const FGeometry& AllottedGeometr
 			if (LabelColorGradient > 0.0f)
 			{
 				const FLinearColor LabelForegroundColor(1.0f, 1.0f, 1.0f, Color.A * LabelColorGradient);
+				const FLinearColor LabelShadowColor(0, 0, 0, Color.A * LabelColorGradient);
+
+				FSlateDrawElement::MakeText(
+					OutDrawElements,
+					++LayerId,
+					AllottedGeometry.ToPaintGeometry(LabelTextPos + FVector2D(2, 2), FVector2D(1, 1)),
+					Label,
+					Font,
+					ESlateDrawEffect::None,
+					LabelShadowColor
+				);
 
 				FSlateDrawElement::MakeText(
 					OutDrawElements,
@@ -988,29 +1100,53 @@ uint32 SWorldPartitionEditorGrid2D::PaintActors(const FGeometry& AllottedGeometr
 						const FPaintGeometry ActorGeometry = AllottedGeometry.ToPaintGeometry(TopLeft, BottomRight - TopLeft);
 						const float LoaderColorGradient = FMath::Min((ActorViewBox.GetArea() - MinimumAreaCull) / AreaFadeDistance, 1.0f);
 						const FLinearColor LoaderColor = LoaderAdapter->GetColor().IsSet() ? *LoaderAdapter->GetColor() : FColor::White;
+						const bool bIsLocalHovered = LocalHoveredLoaderAdapter == LoaderAdapter;
 
 						// Highlight
 						{
 							const FSlateColorBrush LoadedBrush(FLinearColor::White);
 							const FSlateColorBrush UnloadedBrush(FLinearColor::Gray);
-							const FLinearColor LoadedColor(LoaderColor.R, LoaderColor.G, LoaderColor.B, 0.25f * LoaderColorGradient * FullScreenColorGradient);
+							const FLinearColor LoadedColor(LoaderColor.R, LoaderColor.G, LoaderColor.B, 0.23f * LoaderColorGradient * FullScreenColorGradient);
 							const FLinearColor UnloadedColor(LoaderColor.R * 0.15f, LoaderColor.G * 0.15f, LoaderColor.B * 0.15f, 0.25f * LoaderColorGradient * FullScreenColorGradient);
 
-							FSlateDrawElement::MakeBox(
-								OutDrawElements,
-								++LayerId,
-								ActorGeometry,
-								LoaderAdapter->IsLoaded() ? &LoadedBrush : &UnloadedBrush,
-								ESlateDrawEffect::None,
-								LoaderAdapter->IsLoaded() ? LoadedColor : UnloadedColor
-							);
+							// Clip to minimap
+							F2DRectBooleanSubtract LoaderInterfaceHighlight(ActorViewBox);
+
+							if (!bIsLocalHovered)
+							{
+								if (UTexture2D* Texture2D = Cast<UTexture2D>(WorldMiniMapBrush.GetResourceObject()))
+								{
+									const FBox2D MinimapBounds(
+										WorldToScreen.TransformPoint(WorldMiniMapBounds.Min),
+										WorldToScreen.TransformPoint(WorldMiniMapBounds.Max)
+									);
+	
+									LoaderInterfaceHighlight.SubRect(MinimapBounds);
+								}
+							}
+
+							for (const FBox2D& Rect : LoaderInterfaceHighlight.GetRects())
+							{
+								const FPaintGeometry RectAreaGeometry = AllottedGeometry.ToPaintGeometry(
+									Rect.Min,
+									Rect.Max - Rect.Min
+								);
+
+								FSlateDrawElement::MakeBox(
+									OutDrawElements,
+									++LayerId,
+									RectAreaGeometry,
+									LoaderAdapter->IsLoaded() ? &LoadedBrush : &UnloadedBrush,
+									ESlateDrawEffect::None,
+									LoaderAdapter->IsLoaded() ? LoadedColor : UnloadedColor
+								);
+							}
 						}
 
 						// Outline
 						{
 							const bool bIsSelected = SelectedLoaderInterfaces.Contains(LoaderInterface);
 							const bool bIsInsideSelection = (SelectBoxGridSnapped.GetVolume() > 0) && SelectBoxGridSnapped.Intersect(AdapterBounds);
-							const bool bIsLocalHovered = LocalHoveredLoaderAdapter == LoaderAdapter;
 							
 							float OutlineThickness = 2.0f;
 							FLinearColor OutlineColor = FLinearColor::Gray;
@@ -1422,6 +1558,133 @@ int32 SWorldPartitionEditorGrid2D::PaintSoftwareCursor(const FGeometry& Allotted
 	return LayerId + 1;
 }
 
+int32 SWorldPartitionEditorGrid2D::PaintMinimap(const FGeometry& AllottedGeometry, const FSlateRect& MyCullingRect, FSlateWindowElementList& OutDrawElements, int32 LayerId) const
+{
+	// Draw MiniMap image if any
+	if (UTexture2D* Texture2D = Cast<UTexture2D>(WorldMiniMapBrush.GetResourceObject()))
+	{
+		const FBox2D MinimapBounds(
+			WorldToScreen.TransformPoint(WorldMiniMapBounds.Min),
+			WorldToScreen.TransformPoint(WorldMiniMapBounds.Max)
+		);
+
+		const FPaintGeometry WorldImageGeometry = AllottedGeometry.ToPaintGeometry(
+			MinimapBounds.Min,
+			MinimapBounds.Max - MinimapBounds.Min
+		);
+
+		FSlateDrawElement::MakeBox(
+			OutDrawElements,
+			++LayerId,
+			WorldImageGeometry,
+			&WorldMiniMapBrush,
+			ESlateDrawEffect::None,
+			FLinearColor::White
+		);
+
+		F2DRectBooleanSubtract ShadowAreas(MinimapBounds);
+		for (const FLoaderInterface& LoaderInterface : ShownLoaderInterfaces)
+		{
+			if (LoaderInterface.IsValid())
+			{
+				const IWorldPartitionActorLoaderInterface::ILoaderAdapter* LoaderAdapter = LoaderInterface->GetLoaderAdapter();
+
+				if (LoaderAdapter->GetBoundingBox().IsSet())
+				{
+					const FBox AdapterBounds = *LoaderAdapter->GetBoundingBox();
+
+					FVector Origin, Extent;
+					AdapterBounds.GetCenterAndExtents(Origin, Extent);
+
+					const FVector2D TopLeftW = FVector2D(Origin - Extent);
+					const FVector2D BottomRightW = FVector2D(Origin + Extent);
+					const FVector2D TopRightW = FVector2D(BottomRightW.X, TopLeftW.Y);
+					const FVector2D BottomLeftW = FVector2D(TopLeftW.X, BottomRightW.Y);
+
+					const FVector2D TopLeft = WorldToScreen.TransformPoint(TopLeftW);
+					const FVector2D BottomRight = WorldToScreen.TransformPoint(BottomRightW);
+					const FVector2D TopRight = WorldToScreen.TransformPoint(TopRightW);
+					const FVector2D BottomLeft = WorldToScreen.TransformPoint(BottomLeftW);
+
+					const FBox2D ActorViewBox(TopLeft, BottomRight);
+
+					ShadowAreas.SubRect(ActorViewBox);
+				}
+			}
+		}
+
+		for (const FBox2D& ShadowArea : ShadowAreas.GetRects())
+		{
+			const FPaintGeometry ShadowAreaGeometry = AllottedGeometry.ToPaintGeometry(
+				ShadowArea.Min,
+				ShadowArea.Max - ShadowArea.Min
+			);
+
+			const FSlateColorBrush ShadowdBrush(FLinearColor::Gray);
+			const FLinearColor ShadowColor(0, 0, 0, 0.75f);
+
+			FSlateDrawElement::MakeBox(
+				OutDrawElements,
+				++LayerId,
+				ShadowAreaGeometry,
+				&ShadowdBrush,
+				ESlateDrawEffect::None,
+				ShadowColor
+			);
+		}
+
+		// Minimap outline
+		{
+			const TArray<FVector2D> LinePoints(
+			{
+				MinimapBounds.Min,
+				FVector2D(MinimapBounds.Max.X, MinimapBounds.Min.Y), 
+				MinimapBounds.Max, 
+				FVector2D(MinimapBounds.Min.X, 	MinimapBounds.Max.Y), 
+				MinimapBounds.Min
+			});
+
+			FSlateDrawElement::MakeLines
+			(
+				OutDrawElements,
+				++LayerId,
+				AllottedGeometry.ToPaintGeometry(),
+				LinePoints,
+				ESlateDrawEffect::None,
+				FLinearColor::Black
+			);
+		}
+
+		if (Texture2D->IsCurrentlyVirtualTextured())
+		{
+			FVirtualTexture2DResource* VTResource = static_cast<FVirtualTexture2DResource*>(Texture2D->GetResource());
+			const FVector2D ViewportSize = AllottedGeometry.GetLocalSize();
+			const FVector2D ScreenSpaceSize = WorldImageGeometry.GetLocalSize();
+			const FVector2D ViewportPositon = -WorldImageGeometry.GetAccumulatedRenderTransform().GetTranslation() + AllottedGeometry.GetAbsolutePosition();
+
+			FBox2D UVRegion = WorldMiniMapBrush.GetUVRegion();
+			const FVector2D UV0 = UVRegion.Min;
+			const FVector2D UV1 = UVRegion.Max;
+
+			const ERHIFeatureLevel::Type InFeatureLevel = GMaxRHIFeatureLevel;
+			const int32 MipLevel = -1;
+
+			ENQUEUE_RENDER_COMMAND(MakeTilesResident)(
+				[InFeatureLevel, VTResource, ScreenSpaceSize, ViewportPositon, ViewportSize, UV0, UV1, MipLevel](FRHICommandListImmediate& RHICmdList)
+			{
+				// AcquireAllocatedVT() must happen on render thread
+				IAllocatedVirtualTexture* AllocatedVT = VTResource->AcquireAllocatedVT();
+
+				IRendererModule& RenderModule = GetRendererModule();
+				RenderModule.RequestVirtualTextureTilesForRegion(AllocatedVT, ScreenSpaceSize, ViewportPositon, ViewportSize, UV0, UV1, MipLevel);
+				RenderModule.LoadPendingVirtualTextureTiles(RHICmdList, InFeatureLevel);
+			});
+		}
+	}
+
+	return LayerId;
+}
+
 FReply SWorldPartitionEditorGrid2D::FocusSelection()
 {
 	FBox SelectionBox(ForceInit);
@@ -1459,6 +1722,43 @@ FReply SWorldPartitionEditorGrid2D::FocusSelection()
 	return FReply::Handled();
 }
 
+FReply SWorldPartitionEditorGrid2D::FocusLoadedRegions()
+{
+	FBox SelectionBox(ForceInit);
+
+	for (UWorldPartitionEditorLoaderAdapter* EditorLoaderAdapter : WorldPartition->GetRegisteredEditorLoaderAdapters())
+	{
+		check(EditorLoaderAdapter);
+		IWorldPartitionActorLoaderInterface::ILoaderAdapter* LoaderAdapter = EditorLoaderAdapter->GetLoaderAdapter();
+		check(LoaderAdapter);
+
+		if (LoaderAdapter->GetBoundingBox().IsSet() && LoaderAdapter->IsLoaded())
+		{
+			SelectionBox += *LoaderAdapter->GetBoundingBox();
+		}
+	}
+
+	for (TActorIterator<AActor> ActorIt(World); ActorIt; ++ActorIt)
+	{
+		if (AActor* Actor = *ActorIt)
+		{
+			if (Actor->Implements<UWorldPartitionActorLoaderInterface>())
+			{
+				if (IWorldPartitionActorLoaderInterface::ILoaderAdapter* LoaderAdapter = Cast<IWorldPartitionActorLoaderInterface>(Actor)->GetLoaderAdapter())
+				{
+					if (LoaderAdapter->GetBoundingBox().IsSet() && LoaderAdapter->IsLoaded())
+					{
+						SelectionBox += *LoaderAdapter->GetBoundingBox();
+					}
+				}
+			}
+		}
+	}
+
+	FocusBox(SelectionBox);
+	return FReply::Handled();
+}
+
 void SWorldPartitionEditorGrid2D::FocusBox(const FBox& Box) const
 {
 	check(ScreenRect.bIsValid);
@@ -1470,7 +1770,7 @@ void SWorldPartitionEditorGrid2D::FocusBox(const FBox& Box) const
 	{
 		const FVector2D ScreenExtent = ScreenRect.GetExtent();
 		const FVector2D SelectExtent = FVector2D(Box2D.GetExtent());
-		Scale = (ScreenExtent / SelectExtent).GetMin() * 0.75f;
+		Scale = (ScreenExtent / SelectExtent).GetMin() * 0.85f;
 	}
 
 	UpdateTransform();
diff --git a/Engine/Source/Editor/WorldPartitionEditor/Private/WorldPartition/SWorldPartitionEditorGridSpatialHash.cpp b/Engine/Source/Editor/WorldPartitionEditor/Private/WorldPartition/SWorldPartitionEditorGridSpatialHash.cpp
--- a/Engine/Source/Editor/WorldPartitionEditor/Private/WorldPartition/SWorldPartitionEditorGridSpatialHash.cpp
+++ b/Engine/Source/Editor/WorldPartitionEditor/Private/WorldPartition/SWorldPartitionEditorGridSpatialHash.cpp
@@ -1,33 +1,10 @@
 // Copyright Epic Games, Inc. All Rights Reserved.
 #include "WorldPartition/SWorldPartitionEditorGridSpatialHash.h"
-#include "GameFramework/WorldSettings.h"
-#include "DerivedDataCache/Public/DerivedDataCacheInterface.h"
-#include "Misc/SecureHash.h"
-#include "WorldPartition/WorldPartitionRuntimeSpatialHash.h"
-#include "WorldPartition/WorldPartitionEditorSpatialHash.h"
-#include "WorldPartition/WorldPartitionActorDesc.h"
-#include "WorldPartition/WorldPartitionMiniMap.h"
-#include "WorldPartition/WorldPartitionMiniMapHelper.h"
 #include "WorldPartition/WorldPartitionEditorPerProjectUserSettings.h"
-#include "Framework/Commands/Commands.h"
-#include "Framework/MultiBox/MultiBoxBuilder.h"
-#include "Misc/FileHelper.h"
-#include "Misc/ScopedSlowTask.h"
-#include "Layout/WidgetPath.h"
-#include "Framework/Application/SlateApplication.h"
-#include "Slate/SlateTextures.h"
+#include "WorldPartition/WorldPartitionEditorSpatialHash.h"
+#include "WorldPartition/WorldPartition.h"
 #include "Brushes/SlateColorBrush.h"
-#include "ObjectTools.h"
-#include "UObject/ObjectResource.h"
-#include "Styling/AppStyle.h"
-#include "Selection.h"
-#include "Engine/World.h"
-#include "Engine/Texture2D.h"
-#include "Engine/Texture2DDynamic.h"
-#include "Misc/HashBuilder.h"
 #include "Fonts/FontMeasure.h"
-#include "EngineModule.h"
-#include "Renderer/Private/RendererModule.h"
 
 #define LOCTEXT_NAMESPACE "WorldPartitionEditor"
 
@@ -36,42 +13,11 @@ static inline FBox2D ToBox2D(const FBox& Box)
 	return FBox2D(FVector2D(Box.Min), FVector2D(Box.Max));
 }
 
-SWorldPartitionEditorGridSpatialHash::SWorldPartitionEditorGridSpatialHash()
-	: WorldMiniMapBounds(ForceInit)
-{}
-
-SWorldPartitionEditorGridSpatialHash::~SWorldPartitionEditorGridSpatialHash()
-{}
-
 void SWorldPartitionEditorGridSpatialHash::Construct(const FArguments& InArgs)
 {
-	World = InArgs._InWorld;
-	WorldPartition = World ? World->GetWorldPartition() : nullptr;
-
-	if (WorldPartition)
-	{
-		UpdateWorldMiniMapDetails();
-
-		bShowActors = !WorldMiniMapBrush.HasUObject();
-	}
-
 	SWorldPartitionEditorGrid2D::Construct(SWorldPartitionEditorGrid::FArguments().InWorld(InArgs._InWorld));
 }
 
-void SWorldPartitionEditorGridSpatialHash::UpdateWorldMiniMapDetails()
-{
-	if (AWorldPartitionMiniMap* WorldMiniMap = FWorldPartitionMiniMapHelper::GetWorldPartitionMiniMap(World))
-	{
-		WorldMiniMapBounds = FBox2D(FVector2D(WorldMiniMap->MiniMapWorldBounds.Min), FVector2D(WorldMiniMap->MiniMapWorldBounds.Max));
-		if (UTexture2D* MiniMapTexture = WorldMiniMap->MiniMapTexture)
-		{
-			WorldMiniMapBrush.SetUVRegion(WorldMiniMap->UVOffset);
-			WorldMiniMapBrush.SetImageSize(MiniMapTexture->GetImportedSize());
-			WorldMiniMapBrush.SetResourceObject(MiniMapTexture);
-		}
-	}
-}
-
 int32 SWorldPartitionEditorGridSpatialHash::GetSelectionSnap() const
 {
 	UWorldPartitionEditorSpatialHash* EditorSpatialHash = (UWorldPartitionEditorSpatialHash*)WorldPartition->EditorHash;
@@ -135,48 +81,10 @@ int32 SWorldPartitionEditorGridSpatialHash::PaintGrid(const FGeometry& AllottedG
 		);
 	}
 
-	// Draw MiniMap image if any
-	if (UTexture2D* Texture2D = Cast<UTexture2D>(WorldMiniMapBrush.GetResourceObject()))
-	{	
-		FPaintGeometry WorldImageGeometry = AllottedGeometry.ToPaintGeometry(
-			WorldToScreen.TransformPoint(WorldMiniMapBounds.Min),
-			WorldToScreen.TransformPoint(WorldMiniMapBounds.Max) - WorldToScreen.TransformPoint(WorldMiniMapBounds.Min)
-		);
-
-		FSlateDrawElement::MakeRotatedBox(
-			OutDrawElements,
-			++LayerId,
-			WorldImageGeometry,
-			&WorldMiniMapBrush
-		);
-
-		if (Texture2D->IsCurrentlyVirtualTextured())
-		{
-			FVirtualTexture2DResource* VTResource = static_cast<FVirtualTexture2DResource*>(Texture2D->GetResource());
-			const FVector2D ViewportSize = AllottedGeometry.GetLocalSize();
-			const FVector2D ScreenSpaceSize = WorldImageGeometry.GetLocalSize();
-			const FVector2D ViewportPositon = -WorldImageGeometry.GetAccumulatedRenderTransform().GetTranslation() + AllottedGeometry.GetAbsolutePosition();
-
-			FBox2D UVRegion = WorldMiniMapBrush.GetUVRegion();
-			const FVector2D UV0 = UVRegion.Min;
-			const FVector2D UV1 = UVRegion.Max;
-
-			const ERHIFeatureLevel::Type InFeatureLevel = GMaxRHIFeatureLevel;
-			const int32 MipLevel = -1;
-
-			ENQUEUE_RENDER_COMMAND(MakeTilesResident)(
-				[InFeatureLevel, VTResource, ScreenSpaceSize, ViewportPositon, ViewportSize, UV0, UV1, MipLevel](FRHICommandListImmediate& RHICmdList)
-			{
-				// AcquireAllocatedVT() must happen on render thread
-				IAllocatedVirtualTexture* AllocatedVT = VTResource->AcquireAllocatedVT();
-
-				IRendererModule& RenderModule = GetRendererModule();
-				RenderModule.RequestVirtualTextureTilesForRegion(AllocatedVT, ScreenSpaceSize, ViewportPositon, ViewportSize, UV0, UV1, MipLevel);
-				RenderModule.LoadPendingVirtualTextureTiles(RHICmdList, InFeatureLevel);
-			});
-		}
-	}
+	// Paint minimap
+	LayerId = PaintMinimap(AllottedGeometry, MyCullingRect, OutDrawElements, LayerId);
 
+	// Paint grid lines
 	if (ToBox2D(VisibleGridRectWorld).GetArea() > 0.0f)
 	{
 		const FLinearColor Color = FLinearColor(0.1f, 0.1f, 0.1f, 1.f);
@@ -203,7 +111,7 @@ int32 SWorldPartitionEditorGridSpatialHash::PaintGrid(const FGeometry& AllottedG
 			LinePoints[0] = WorldToScreen.TransformPoint(LineStartH);
 			LinePoints[1] = WorldToScreen.TransformPoint(LineEndH);
 
-			FSlateDrawElement::MakeLines(OutDrawElements, LayerId, AllottedGeometry.ToPaintGeometry(), LinePoints, ESlateDrawEffect::NoBlending, Color, false, 1.0f);
+			FSlateDrawElement::MakeLines(OutDrawElements, LayerId, AllottedGeometry.ToPaintGeometry(), LinePoints, ESlateDrawEffect::None, Color, false, 1.0f);
 		}
 
 		// Vertical
@@ -215,7 +123,7 @@ int32 SWorldPartitionEditorGridSpatialHash::PaintGrid(const FGeometry& AllottedG
 			LinePoints[0] = WorldToScreen.TransformPoint(LineStartH);
 			LinePoints[1] = WorldToScreen.TransformPoint(LineEndH);
 
-			FSlateDrawElement::MakeLines(OutDrawElements, LayerId, AllottedGeometry.ToPaintGeometry(), LinePoints, ESlateDrawEffect::NoBlending, Color, false, 1.0f);
+			FSlateDrawElement::MakeLines(OutDrawElements, LayerId, AllottedGeometry.ToPaintGeometry(), LinePoints, ESlateDrawEffect::None, Color, false, 1.0f);
 		}
 
 		// Draw coordinates
diff --git a/Engine/Source/Editor/WorldPartitionEditor/Public/WorldPartition/SWorldPartitionEditorGrid2D.h b/Engine/Source/Editor/WorldPartitionEditor/Public/WorldPartition/SWorldPartitionEditorGrid2D.h
--- a/Engine/Source/Editor/WorldPartitionEditor/Public/WorldPartition/SWorldPartitionEditorGrid2D.h
+++ b/Engine/Source/Editor/WorldPartitionEditor/Public/WorldPartition/SWorldPartitionEditorGrid2D.h
@@ -64,8 +64,10 @@ protected:
 	virtual uint32 PaintViewer(const FGeometry& AllottedGeometry, const FSlateRect& MyCullingRect, FSlateWindowElementList& OutDrawElements, uint32 LayerId) const;
 	virtual uint32 PaintSelection(const FGeometry& AllottedGeometry, const FSlateRect& MyCullingRect, FSlateWindowElementList& OutDrawElements, uint32 LayerId) const;
 	virtual int32 PaintSoftwareCursor(const FGeometry& AllottedGeometry, const FSlateRect& MyCullingRect, FSlateWindowElementList& OutDrawElements, int32 LayerId) const;
+	virtual int32 PaintMinimap(const FGeometry& AllottedGeometry, const FSlateRect& MyCullingRect, FSlateWindowElementList& OutDrawElements, int32 LayerId) const;
 
 	virtual FReply FocusSelection();
+	virtual FReply FocusLoadedRegions();
 
 	void UpdateTransform() const;
 	void UpdateSelectionBox(bool bSnap);
@@ -140,6 +142,11 @@ protected:
 	FLoaderInterfaceStack HoveredLoaderInterfacesStack;
 	FLoaderInterface HoveredLoaderInterface;
 
+	// Minimap
+	void UpdateWorldMiniMapDetails();
+	FBox2D WorldMiniMapBounds;
+	FSlateBrush WorldMiniMapBrush;
+
 	// Profiling
 	mutable double TickTime;
 	mutable double PaintTime;
diff --git a/Engine/Source/Editor/WorldPartitionEditor/Public/WorldPartition/SWorldPartitionEditorGridSpatialHash.h b/Engine/Source/Editor/WorldPartitionEditor/Public/WorldPartition/SWorldPartitionEditorGridSpatialHash.h
--- a/Engine/Source/Editor/WorldPartitionEditor/Public/WorldPartition/SWorldPartitionEditorGridSpatialHash.h
+++ b/Engine/Source/Editor/WorldPartitionEditor/Public/WorldPartition/SWorldPartitionEditorGridSpatialHash.h
@@ -3,24 +3,14 @@
 
 #include "CoreMinimal.h"
 #include "SWorldPartitionEditorGrid2D.h"
-#include "Brushes/SlateDynamicImageBrush.h"
-#include "Framework/Commands/UICommandList.h"
 
 class SWorldPartitionEditorGridSpatialHash : public SWorldPartitionEditorGrid2D
 {
 public:
 	WORLD_PARTITION_EDITOR_IMPL(SWorldPartitionEditorGridSpatialHash);
 
-	SWorldPartitionEditorGridSpatialHash();
-	~SWorldPartitionEditorGridSpatialHash();
-
 	void Construct(const FArguments& InArgs);
 
 	virtual int32 GetSelectionSnap() const override;
 	virtual int32 PaintGrid(const FGeometry& AllottedGeometry, const FSlateRect& MyCullingRect, FSlateWindowElementList& OutDrawElements, int32 LayerId) const override;
-
-private:
-	void UpdateWorldMiniMapDetails();
-	FBox2D WorldMiniMapBounds;
-	FSlateBrush WorldMiniMapBrush;
 };
