Added missing module dependency.

diff --git a/Engine/Source/Editor/WorldPartitionEditor/WorldPartitionEditor.Build.cs b/Engine/Source/Editor/WorldPartitionEditor/WorldPartitionEditor.Build.cs
--- a/Engine/Source/Editor/WorldPartitionEditor/WorldPartitionEditor.Build.cs
+++ b/Engine/Source/Editor/WorldPartitionEditor/WorldPartitionEditor.Build.cs
@@ -37,6 +37,7 @@ public class WorldPartitionEditor : ModuleRules
 		(
 			new string[]
 			{
+				"SceneOutliner",
 				"WorkspaceMenuStructure",
 			}
 		);
