World Partition Editor: add a drop shadow when showing actor bounds.

#rnx

diff --git a/Engine/Source/Editor/WorldPartitionEditor/Private/WorldPartition/SWorldPartitionEditorGrid2D.cpp b/Engine/Source/Editor/WorldPartitionEditor/Private/WorldPartition/SWorldPartitionEditorGrid2D.cpp
--- a/Engine/Source/Editor/WorldPartitionEditor/Private/WorldPartition/SWorldPartitionEditorGrid2D.cpp
+++ b/Engine/Source/Editor/WorldPartitionEditor/Private/WorldPartition/SWorldPartitionEditorGrid2D.cpp
@@ -1231,9 +1231,10 @@ uint32 SWorldPartitionEditorGrid2D::PaintActors(const FGeometry& AllottedGeometr
 				const UClass* ActorClass = ActorDescView.GetActorNativeClass();
 				const AActor* Actor = ActorDescView.GetActor();
 				const FPaintGeometry ActorGeometry = AllottedGeometry.ToPaintGeometry(TopLeft, BottomRight - TopLeft);
+				const FPaintGeometry ActorGeometryShadow = AllottedGeometry.ToPaintGeometry(TopLeft - FVector2D::One(), BottomRight - TopLeft + FVector2D::One() * 2.0);
 				const bool bIsSelected = Actor ? Actor->IsSelected() : false;
 				const bool bIsBoundsEdgeHovered = !bIsBoundsEdgeHoveredFound && IsBoundsEdgeHovered(MouseCursorPos, ActorViewBox, 10.0f);
-				const float ActorColorGradient = FMath::Min((ActorViewBox.GetArea() - MinimumAreaCull) / AreaFadeDistance, 1.0f);
+				const float ActorColorGradient = FMath::Cube(FMath::Clamp((ActorViewBox.GetArea() - MinimumAreaCull) / AreaFadeDistance, 0.0f, 1.0f));
 				const float ActorBrightness = ActorDescView.GetIsSpatiallyLoaded() ? 1.0f : 0.3f;
 
 				// Only draw the first edge bounds hovered
@@ -1249,13 +1250,22 @@ uint32 SWorldPartitionEditorGrid2D::PaintActors(const FGeometry& AllottedGeometr
 						DrawActorLabel(ActorLabel.ToString(), bIsBoundsEdgeHovered ? MouseCursorPos : ActorViewBox.GetCenter(), ActorGeometry, FLinearColor::Yellow, SmallLayoutFont, true);
 					}
 
-					ActorColor = FLinearColor::Yellow;
+					ActorColor = FLinearColor::Yellow.CopyWithNewOpacity(ActorColor.A);
 				}
 				else if ((SelectBoxGridSnapped.GetVolume() > 0) && SelectBoxGridSnapped.Intersect(ActorDescView.GetBounds()))
 				{
-					ActorColor = FLinearColor::White;
+					ActorColor = FLinearColor::White.CopyWithNewOpacity(ActorColor.A);
 				}
 
+				FSlateDrawElement::MakeBox(
+					OutDrawElements,
+					++LayerId,
+					ActorGeometryShadow,
+					FAppStyle::GetBrush(TEXT("Border")),
+					ESlateDrawEffect::None,
+					FLinearColor::Black.CopyWithNewOpacity(ActorColor.A)
+				);
+
 				FSlateDrawElement::MakeBox(
 					OutDrawElements,
 					++LayerId,
