Fix missing call to Super::BeginDestroy

#rnx

diff --git a/Engine/Source/Editor/WorldPartitionEditor/Public/WorldPartition/ContentBundle/ContentBundleEditorSubsystem.h b/Engine/Source/Editor/WorldPartitionEditor/Public/WorldPartition/ContentBundle/ContentBundleEditorSubsystem.h
--- a/Engine/Source/Editor/WorldPartitionEditor/Public/WorldPartition/ContentBundle/ContentBundleEditorSubsystem.h
+++ b/Engine/Source/Editor/WorldPartitionEditor/Public/WorldPartition/ContentBundle/ContentBundleEditorSubsystem.h
@@ -21,7 +21,7 @@ class UContentBundleEditorSubsystemModule : public UObject
 
 public:
 	//~ Begin UObject interface
-	virtual void BeginDestroy() override { check(!IsInitialize()); }
+	virtual void BeginDestroy() override { check(!IsInitialize()); Super::BeginDestroy(); }
 	//~ End UObject interface
 
 	void Initialize() { DoInitialize(); bIsInitialized = true; }
