World Partition Editor: added play from here option (mapped on shift+double click).

#rnx

diff --git a/Engine/Source/Editor/WorldPartitionEditor/Private/WorldPartition/SWorldPartitionEditorGrid2D.cpp b/Engine/Source/Editor/WorldPartitionEditor/Private/WorldPartition/SWorldPartitionEditorGrid2D.cpp
--- a/Engine/Source/Editor/WorldPartitionEditor/Private/WorldPartition/SWorldPartitionEditorGrid2D.cpp
+++ b/Engine/Source/Editor/WorldPartitionEditor/Private/WorldPartition/SWorldPartitionEditorGrid2D.cpp
@@ -30,6 +30,9 @@
 #include "LocationVolume.h"
 #include "SWorldPartitionViewportWidget.h"
 
+#include "IAssetViewport.h"
+#include "Kismet2/DebuggerCommands.h"
+
 #define LOCTEXT_NAMESPACE "WorldPartitionEditor"
 
 static bool GShowEditorProfilingStats = 0;
@@ -231,6 +234,7 @@ void SWorldPartitionEditorGrid2D::FEditorCommands::RegisterCommands()
 	UI_COMMAND(UnloadSelectedRegions, "Unload Selected Regions", "Unload the selected regions.", EUserInterfaceActionType::Button, FInputChord());
 	UI_COMMAND(ConvertSelectedRegionsToActors, "Convert Selected Regions To Actors", "Convert the selected regions to actors.", EUserInterfaceActionType::Button, FInputChord());
 	UI_COMMAND(MoveCameraHere, "Move Camera Here", "Move the camera to the selected position.", EUserInterfaceActionType::Button, FInputChord());
+	UI_COMMAND(PlayFromHere, "Play From Here", "Play from here.", EUserInterfaceActionType::Button, FInputChord());
 }
 
 SWorldPartitionEditorGrid2D::SWorldPartitionEditorGrid2D()
@@ -437,6 +441,7 @@ void SWorldPartitionEditorGrid2D::Construct(const FArguments& InArgs)
 	ActionList.MapAction(Commands.UnloadSelectedRegions, FExecuteAction::CreateSP(this, &SWorldPartitionEditorGrid2D::UnloadSelectedRegions), FCanExecuteAction::CreateLambda(CanUnloadSelectedRegions));
 	ActionList.MapAction(Commands.ConvertSelectedRegionsToActors, FExecuteAction::CreateSP(this, &SWorldPartitionEditorGrid2D::ConvertSelectedRegionsToActors), FCanExecuteAction::CreateLambda(CanConvertSelectedRegionsToActors));
 	ActionList.MapAction(Commands.MoveCameraHere, FExecuteAction::CreateSP(this, &SWorldPartitionEditorGrid2D::MoveCameraHere));
+	ActionList.MapAction(Commands.PlayFromHere, FExecuteAction::CreateSP(this, &SWorldPartitionEditorGrid2D::PlayFromHere));
 }
 
 void SWorldPartitionEditorGrid2D::UpdateWorldMiniMapDetails()
@@ -570,6 +575,29 @@ void SWorldPartitionEditorGrid2D::MoveCameraHere()
 	}
 }
 
+void SWorldPartitionEditorGrid2D::PlayFromHere()
+{
+	FVector StartLocation = FVector(MouseCursorPosWorld, 0);
+	FRotator StartRotation;
+
+	FHitResult HitResult;
+	const FVector TraceStart(StartLocation.X, StartLocation.Y, HALF_WORLD_MAX);
+	const FVector TraceEnd(StartLocation.X, StartLocation.Y, -HALF_WORLD_MAX);
+	const FCollisionQueryParams TraceParams(SCENE_QUERY_STAT(SWorldPartitionEditorGrid2D_MoveCameraHere), true);
+	const bool bHitResultValid = World->LineTraceSingleByChannel(HitResult, TraceStart, TraceEnd, ECC_WorldStatic, TraceParams);
+
+	FLevelEditorModule& LevelEditorModule = FModuleManager::GetModuleChecked<FLevelEditorModule>("LevelEditor");
+	TSharedPtr<IAssetViewport> ActiveLevelViewport = LevelEditorModule.GetFirstActiveViewport();
+
+	if (ActiveLevelViewport.IsValid() && ActiveLevelViewport->GetAssetViewportClient().IsPerspective())
+	{
+		StartLocation.Z = bHitResultValid ? HitResult.ImpactPoint.Z : ActiveLevelViewport->GetAssetViewportClient().GetViewLocation().Z;
+		StartRotation = ActiveLevelViewport->GetAssetViewportClient().GetViewRotation();
+
+		FPlayWorldCommandCallbacks::StartPlayFromHere(StartLocation, StartRotation, ActiveLevelViewport);
+	}	
+}
+
 bool SWorldPartitionEditorGrid2D::IsFollowPlayerInPIE() const
 {
 	return bFollowPlayerInPIE && (UWorldPartition::IsSimulating() || GEditor->PlayWorld);
@@ -662,6 +690,7 @@ FReply SWorldPartitionEditorGrid2D::OnMouseButtonUp(const FGeometry& MyGeometry,
 
 			MenuBuilder.BeginSection(NAME_None, LOCTEXT("WorldPartitionMisc", "Misc"));
 				MenuBuilder.AddMenuEntry(Commands.MoveCameraHere);
+				MenuBuilder.AddMenuEntry(Commands.PlayFromHere);
 			MenuBuilder.EndSection();			
 
 			FWidgetPath WidgetPath = MouseEvent.GetEventPath() != nullptr ? *MouseEvent.GetEventPath() : FWidgetPath();
@@ -725,7 +754,14 @@ FReply SWorldPartitionEditorGrid2D::OnMouseButtonDoubleClick(const FGeometry& In
 {
 	if (IsInteractive())
 	{
-		MoveCameraHere();
+		if (InMouseEvent.IsShiftDown())
+		{
+			PlayFromHere();
+		}
+		else
+		{
+			MoveCameraHere();
+		}
 	}
 	return FReply::Handled();
 }
diff --git a/Engine/Source/Editor/WorldPartitionEditor/Public/WorldPartition/SWorldPartitionEditorGrid2D.h b/Engine/Source/Editor/WorldPartitionEditor/Public/WorldPartition/SWorldPartitionEditorGrid2D.h
--- a/Engine/Source/Editor/WorldPartitionEditor/Public/WorldPartition/SWorldPartitionEditorGrid2D.h
+++ b/Engine/Source/Editor/WorldPartitionEditor/Public/WorldPartition/SWorldPartitionEditorGrid2D.h
@@ -24,6 +24,7 @@ protected:
 		TSharedPtr<FUICommandInfo> UnloadHoveredRegion;
 		TSharedPtr<FUICommandInfo> ConvertSelectedRegionsToActors;
 		TSharedPtr<FUICommandInfo> MoveCameraHere;
+		TSharedPtr<FUICommandInfo> PlayFromHere;
 
 		/**
 		 * Initialize commands
@@ -43,6 +44,7 @@ protected:
 	void UnloadSelectedRegions();
 	void ConvertSelectedRegionsToActors();
 	void MoveCameraHere();
+	void PlayFromHere();
 
 	bool IsFollowPlayerInPIE() const;
 	bool IsInteractive() const;
