Add an outliner to display content bundle infos.
The outliner currently only displays content bundle from the editor world.
The outliner displays the client status and the amount of actors in the content bundle
It allows to toggle editing on content bundle and start adding actors
A context menu on the tree items allow to (de)select and (un)load actors of the content bundle

diff --git a/Engine/Source/Editor/DataLayerEditor/DataLayerEditor.Build.cs b/Engine/Source/Editor/DataLayerEditor/DataLayerEditor.Build.cs
--- a/Engine/Source/Editor/DataLayerEditor/DataLayerEditor.Build.cs
+++ b/Engine/Source/Editor/DataLayerEditor/DataLayerEditor.Build.cs
@@ -22,6 +22,7 @@ public class DataLayerEditor : ModuleRules
 				"SceneOutliner",
 				"ToolMenus",
 				"AssetTools",
+				"ContentBrowserData"
 			}
 		);
 	}
diff --git a/Engine/Source/Editor/WorldPartitionEditor/Private/WorldPartition/ContentBundle/Outliner/ContentBundleActorCountColumn.cpp b/Engine/Source/Editor/WorldPartitionEditor/Private/WorldPartition/ContentBundle/Outliner/ContentBundleActorCountColumn.cpp
new file mode 100644
--- /dev/null
+++ b/Engine/Source/Editor/WorldPartitionEditor/Private/WorldPartition/ContentBundle/Outliner/ContentBundleActorCountColumn.cpp
@@ -0,0 +1,74 @@
+// Copyright Epic Games, Inc. All Rights Reserved.
+
+#include "WorldPartition/ContentBundle/Outliner/ContentBundleActorCountColumn.h"
+
+#include "Widgets/Images/SImage.h"
+#include "Widgets/Input/SButton.h"
+#include "Widgets/Input/SCheckBox.h"
+#include "Widgets/Views/STreeView.h"
+#include "ScopedTransaction.h"
+#include "Editor.h"
+#include "Styling/AppStyle.h"
+#include "WorldPartition/ContentBundle/Outliner/ContentBundleTreeItem.h"
+#include "WorldPartition/ContentBundle/ContentBundleEditor.h"
+
+#define LOCTEXT_NAMESPACE "ContentBundle"
+
+namespace ContentBundleOutlinerPrivate
+{
+	FName ContentBundleOutlinerBundleActorCount("Content Bundle Actor Count");
+}
+
+FName FContentBundleOutlinerActorCountColumn::GetID()
+{
+	return ContentBundleOutlinerPrivate::ContentBundleOutlinerBundleActorCount;
+}
+
+SHeaderRow::FColumn::FArguments FContentBundleOutlinerActorCountColumn::ConstructHeaderRowColumn()
+{
+	return SHeaderRow::Column(GetColumnID())
+		.HAlignHeader(HAlign_Center)
+		.VAlignHeader(VAlign_Center)
+		.HAlignCell(HAlign_Center)
+		.VAlignCell(VAlign_Center)
+		.DefaultTooltip(LOCTEXT("ContentBundleOutlinerActorCountTooltip", "Actor Count (Unsaved Actor Count)"))
+		[
+			SNew(STextBlock)
+			.Margin(FMargin(0.0f))
+			.Text(LOCTEXT("ContentBundleOutlinerActorCount", "Actor Count"))
+		];
+}
+
+const TSharedRef<SWidget> FContentBundleOutlinerActorCountColumn::ConstructRowWidget(FSceneOutlinerTreeItemRef TreeItem, const STableRow<FSceneOutlinerTreeItemPtr>& Row)
+{
+	if (FContentBundleTreeItem* ContentBundleTreeItem = TreeItem->CastTo<FContentBundleTreeItem>())
+	{
+		TSharedPtr<FContentBundleEditor> ContentBundleEditor = ContentBundleTreeItem->GetContentBundleEditorPin();
+		if (ContentBundleEditor != nullptr)
+		{
+			return SNew(SHorizontalBox)
+				+ SHorizontalBox::Slot()
+				.Padding(0, 0, 0, 0)
+				.AutoWidth()
+				.VAlign(VAlign_Center)
+				[
+					SNew(STextBlock)
+					.Text_Lambda([ContentBundleTreeItem]
+						{
+							TSharedPtr<FContentBundleEditor> ContentBundleEditorPin = ContentBundleTreeItem->GetContentBundleEditorPin();
+							if (ContentBundleEditorPin != nullptr)
+							{
+								FString ActorCountString = FString::Printf(TEXT("%u(%u)"), ContentBundleEditorPin->GetActorCount(), ContentBundleEditorPin->GetUnsavedActorAcount());
+								return FText::FromString(ActorCountString);
+							}
+
+							return FText::AsNumber(0);
+						})
+					.ColorAndOpacity(MakeAttributeLambda([ContentBundleTreeItem] { return ContentBundleTreeItem->GetItemColor(); }))
+				];
+		}
+	}
+	return SNullWidget::NullWidget;
+}
+
+#undef LOCTEXT_NAMESPACE
\ No newline at end of file
diff --git a/Engine/Source/Editor/WorldPartitionEditor/Private/WorldPartition/ContentBundle/Outliner/ContentBundleClientColumn.cpp b/Engine/Source/Editor/WorldPartitionEditor/Private/WorldPartition/ContentBundle/Outliner/ContentBundleClientColumn.cpp
new file mode 100644
--- /dev/null
+++ b/Engine/Source/Editor/WorldPartitionEditor/Private/WorldPartition/ContentBundle/Outliner/ContentBundleClientColumn.cpp
@@ -0,0 +1,112 @@
+// Copyright Epic Games, Inc. All Rights Reserved.
+
+#include "WorldPartition/ContentBundle/Outliner/ContentBundleClientColumn.h"
+
+#include "Widgets/Images/SImage.h"
+#include "Widgets/Input/SButton.h"
+#include "Widgets/Input/SCheckBox.h"
+#include "Widgets/Views/STreeView.h"
+#include "ScopedTransaction.h"
+#include "Editor.h"
+#include "Styling/AppStyle.h"
+#include "WorldPartition/ContentBundle/Outliner/ContentBundleTreeItem.h"
+#include "WorldPartition/ContentBundle/ContentBundleEditor.h"
+#include "WorldPartition/ContentBundle/ContentBundle.h"
+#include "WorldPartition/ContentBundle/ContentBundleClient.h"
+#include "SortHelper.h"
+
+#define LOCTEXT_NAMESPACE "ContentBundle"
+
+namespace ContentBundleOutlinerPrivate
+{
+	FName ContentBundleOutlinerClient("Content Bundle Client");
+}
+
+FName FContentBundleOutlinerClientColumn::GetID()
+{
+	return ContentBundleOutlinerPrivate::ContentBundleOutlinerClient;
+}
+
+SHeaderRow::FColumn::FArguments FContentBundleOutlinerClientColumn::ConstructHeaderRowColumn()
+{
+	return SHeaderRow::Column(GetColumnID())
+		.HAlignHeader(HAlign_Center)
+		.VAlignHeader(VAlign_Center)
+		.HAlignCell(HAlign_Center)
+		.VAlignCell(VAlign_Center)
+		.DefaultTooltip(FText::FromName(GetColumnID()))
+		[
+			SNew(STextBlock)
+			.Margin(FMargin(0.0f))
+			.Text(LOCTEXT("ContentBundleOutlinerClient", "Client"))
+		];
+}
+
+const TSharedRef<SWidget> FContentBundleOutlinerClientColumn::ConstructRowWidget(FSceneOutlinerTreeItemRef TreeItem, const STableRow<FSceneOutlinerTreeItemPtr>& Row)
+{
+	if (FContentBundleTreeItem* ContentBundleTreeItem = TreeItem->CastTo<FContentBundleTreeItem>())
+	{
+		TWeakPtr<FContentBundleEditor> ContentBundleEditor = ContentBundleTreeItem->GetContentBundleEditor();
+		if (ContentBundleEditor != nullptr)
+		{
+			return SNew(SHorizontalBox)
+				+ SHorizontalBox::Slot()
+				.Padding(0, 0, 0, 0)
+				.AutoWidth()
+				.VAlign(VAlign_Center)
+				[
+					SNew(STextBlock)
+					.Text_Lambda([this, ContentBundleEditor] { return GetClientDisplayName(ContentBundleEditor); })
+					.ColorAndOpacity(MakeAttributeLambda([ContentBundleTreeItem] { return ContentBundleTreeItem->GetItemColor(); }))
+				];
+		}
+	}
+	return SNullWidget::NullWidget;
+}
+
+void FContentBundleOutlinerClientColumn::SortItems(TArray<FSceneOutlinerTreeItemPtr>& OutItems, const EColumnSortMode::Type SortMode) const
+{
+	typedef FSceneOutlinerSortHelper<int32, SceneOutliner::FNumericStringWrapper> FSort;
+
+	auto SortByClientName = [this](const ISceneOutlinerTreeItem& Item)
+	{
+		if (const FContentBundleTreeItem* ContentBundleTreeItem = Item.CastTo<FContentBundleTreeItem>())
+		{
+			TSharedPtr<FContentBundleEditor> ContentBundleEditorPin = ContentBundleTreeItem->GetContentBundleEditorPin();
+			if (ContentBundleEditorPin != nullptr)
+			{
+				TSharedPtr<FContentBundleClient> Client = ContentBundleEditorPin->GetClient().Pin();
+				if (Client != nullptr)
+				{
+					FString Temp = Client->GetDisplayName();
+					return SceneOutliner::FNumericStringWrapper(MoveTemp(Temp));
+				}
+			}
+		}
+
+		return SceneOutliner::FNumericStringWrapper();
+	};
+
+	FSort()
+		.Primary([this](const ISceneOutlinerTreeItem& Item) { return WeakSceneOutliner.Pin()->GetTypeSortPriority(Item); }, SortMode)
+		.Secondary(SortByClientName, SortMode)
+		.Sort(OutItems);
+}
+
+FText FContentBundleOutlinerClientColumn::GetClientDisplayName(const TWeakPtr<FContentBundleEditor>& ContentBundleEditor) const
+{
+	TSharedPtr<FContentBundleEditor> ContentBundleEditorPin = ContentBundleEditor.Pin();
+	if (ContentBundleEditorPin != nullptr)
+	{
+		TSharedPtr<FContentBundleClient> ContentBundleClient = ContentBundleEditorPin->GetClient().Pin();
+		if (ContentBundleClient != nullptr)
+		{
+			return FText::FromString(ContentBundleClient->GetDisplayName());
+		}
+	}
+	
+
+	return LOCTEXT("Unknow", "Unknown");
+}
+
+#undef LOCTEXT_NAMESPACE
\ No newline at end of file
diff --git a/Engine/Source/Editor/WorldPartitionEditor/Private/WorldPartition/ContentBundle/Outliner/ContentBundleEditingColumn.cpp b/Engine/Source/Editor/WorldPartitionEditor/Private/WorldPartition/ContentBundle/Outliner/ContentBundleEditingColumn.cpp
new file mode 100644
--- /dev/null
+++ b/Engine/Source/Editor/WorldPartitionEditor/Private/WorldPartition/ContentBundle/Outliner/ContentBundleEditingColumn.cpp
@@ -0,0 +1,74 @@
+// Copyright Epic Games, Inc. All Rights Reserved.
+
+#include "WorldPartition/ContentBundle/Outliner/ContentBundleEditingColumn.h"
+
+#include "Widgets/Images/SImage.h"
+#include "Widgets/Input/SButton.h"
+#include "Widgets/Input/SCheckBox.h"
+#include "Widgets/Views/STreeView.h"
+#include "ScopedTransaction.h"
+#include "Editor.h"
+#include "Styling/AppStyle.h"
+#include "WorldPartition/ContentBundle/Outliner/ContentBundleTreeItem.h"
+#include "WorldPartition/ContentBundle/ContentBundleEditor.h"
+
+#define LOCTEXT_NAMESPACE "ContentBundle"
+
+namespace ContentBundleOutlinerPrivate
+{
+	FName ContentBundleOutlinerEditing("Editing");
+}
+
+FName FContentBundleOutlinerEditingColumn::GetID()
+{
+	return ContentBundleOutlinerPrivate::ContentBundleOutlinerEditing;
+}
+
+SHeaderRow::FColumn::FArguments FContentBundleOutlinerEditingColumn::ConstructHeaderRowColumn()
+{
+	return SHeaderRow::Column(GetColumnID())
+		.HAlignHeader(HAlign_Center)
+		.VAlignHeader(VAlign_Center)
+		.HAlignCell(HAlign_Center)
+		.VAlignCell(VAlign_Center)
+		.DefaultTooltip(FText::FromName(GetColumnID()))
+		[
+			SNew(STextBlock)
+			.Margin(FMargin(0.0f))
+			.Text(LOCTEXT("ContentBundleOutlinerEditing", "Editing"))
+
+		];
+}
+
+const TSharedRef<SWidget> FContentBundleOutlinerEditingColumn::ConstructRowWidget(FSceneOutlinerTreeItemRef TreeItem, const STableRow<FSceneOutlinerTreeItemPtr>& Row)
+{
+	if (FContentBundleTreeItem* ContentBundleTreeItem = TreeItem->CastTo<FContentBundleTreeItem>())
+	{
+		TWeakPtr<FContentBundleEditor> ContentBundleEditor = ContentBundleTreeItem->GetContentBundleEditor();
+		if (ContentBundleEditor != nullptr)
+		{
+			return SNew(SHorizontalBox)
+				+ SHorizontalBox::Slot()
+				.Padding(0, 0, 0, 0)
+				.AutoWidth()
+				.VAlign(VAlign_Center)
+				[
+					SNew(STextBlock)
+					.Text_Lambda([ContentBundleTreeItem] 
+						{ 
+							TSharedPtr<FContentBundleEditor> ContentBundleEditorPin = ContentBundleTreeItem->GetContentBundleEditorPin();
+							if (ContentBundleEditorPin != nullptr && ContentBundleEditorPin->IsBeingEdited())
+							{
+								return LOCTEXT("ContentBundleOutlinerEditingActive", "Active");
+							}
+							
+							return LOCTEXT("", ""); 
+						})
+					.ColorAndOpacity(MakeAttributeLambda([ContentBundleTreeItem] { return ContentBundleTreeItem->GetItemColor(); }))
+				];
+		}
+	}
+	return SNullWidget::NullWidget;
+}
+
+#undef LOCTEXT_NAMESPACE
\ No newline at end of file
diff --git a/Engine/Source/Editor/WorldPartitionEditor/Private/WorldPartition/ContentBundle/Outliner/ContentBundleHasErrorColumn.cpp b/Engine/Source/Editor/WorldPartitionEditor/Private/WorldPartition/ContentBundle/Outliner/ContentBundleHasErrorColumn.cpp
new file mode 100644
--- /dev/null
+++ b/Engine/Source/Editor/WorldPartitionEditor/Private/WorldPartition/ContentBundle/Outliner/ContentBundleHasErrorColumn.cpp
@@ -0,0 +1,77 @@
+// Copyright Epic Games, Inc. All Rights Reserved.
+
+#include "WorldPartition/ContentBundle/Outliner/ContentBundleHasErrorColumn.h"
+
+#include "Widgets/Images/SImage.h"
+#include "Widgets/Input/SButton.h"
+#include "Widgets/Input/SCheckBox.h"
+#include "Widgets/Views/STreeView.h"
+#include "ScopedTransaction.h"
+#include "Editor.h"
+#include "Styling/AppStyle.h"
+#include "WorldPartition/ContentBundle/Outliner/ContentBundleTreeItem.h"
+#include "WorldPartition/ContentBundle/ContentBundleEditor.h"
+
+#define LOCTEXT_NAMESPACE "ContentBundle"
+
+namespace ContentBundleOutlinerPrivate
+{
+	FName ContentBundleOutlinerBundleHasError("Content Bundle Has Errors");
+}
+
+FName FContentBundleOutlinerHasErrorColumn::GetID()
+{
+	return ContentBundleOutlinerPrivate::ContentBundleOutlinerBundleHasError;
+}
+
+SHeaderRow::FColumn::FArguments FContentBundleOutlinerHasErrorColumn::ConstructHeaderRowColumn()
+{
+	return SHeaderRow::Column(GetColumnID())
+		.FixedWidth(24.f)
+		.HAlignHeader(HAlign_Center)
+		.VAlignHeader(VAlign_Center)
+		.HAlignCell(HAlign_Center)
+		.VAlignCell(VAlign_Center)
+		.DefaultTooltip(FText::FromName(GetColumnID()))
+		[
+			SNew(SImage)
+			.Image(FAppStyle::GetBrush(TEXT("Icons.Error")))
+			.ColorAndOpacity(FSlateColor::UseForeground())
+		];
+}
+
+const TSharedRef<SWidget> FContentBundleOutlinerHasErrorColumn::ConstructRowWidget(FSceneOutlinerTreeItemRef TreeItem, const STableRow<FSceneOutlinerTreeItemPtr>& Row)
+{
+	if (FContentBundleTreeItem* ContentBundleTreeItem = TreeItem->CastTo<FContentBundleTreeItem>())
+	{
+		TSharedPtr<FContentBundleEditor> ContentBundleEditor = ContentBundleTreeItem->GetContentBundleEditorPin();
+		if (ContentBundleEditor != nullptr)
+		{
+			if (!ContentBundleEditor->IsValid())
+			{
+				return SNew(SHorizontalBox)
+					+ SHorizontalBox::Slot()
+					.Padding(0, 0, 0, 0)
+					.AutoWidth()
+					.VAlign(VAlign_Center)
+					[
+						SNew(SImage)
+						.Image(FAppStyle::GetBrush(TEXT("Icons.Error")))
+						.ToolTipText(LOCTEXT("ContentBundleOutlinerHasError", "Content Bundle has error. Consult log for details."))
+						.ColorAndOpacity(MakeAttributeLambda([ContentBundleTreeItem] { return ContentBundleTreeItem->GetItemColor(); }))
+					];
+			}
+			else
+			{
+				return SNew(SHorizontalBox)
+					+ SHorizontalBox::Slot()
+					.Padding(0, 0, 0, 0)
+					.AutoWidth()
+					.VAlign(VAlign_Center);
+			}
+		}
+	}
+	return SNullWidget::NullWidget;
+}
+
+#undef LOCTEXT_NAMESPACE
\ No newline at end of file
diff --git a/Engine/Source/Editor/WorldPartitionEditor/Private/WorldPartition/ContentBundle/Outliner/ContentBundleHierarchy.cpp b/Engine/Source/Editor/WorldPartitionEditor/Private/WorldPartition/ContentBundle/Outliner/ContentBundleHierarchy.cpp
new file mode 100644
--- /dev/null
+++ b/Engine/Source/Editor/WorldPartitionEditor/Private/WorldPartition/ContentBundle/Outliner/ContentBundleHierarchy.cpp
@@ -0,0 +1,95 @@
+// Copyright Epic Games, Inc. All Rights Reserved.
+
+#include "WorldPartition/ContentBundle/Outliner/ContentBundleHierarchy.h"
+
+#include "WorldPartition/ContentBundle/Outliner/ContentBundleMode.h"
+#include "WorldPartition/ContentBundle/ContentBundleEditorSubsystem.h"
+#include "WorldPartition/ContentBundle/ContentBundleEditor.h"
+#include "WorldPartition/ContentBundle/Outliner/ContentBundleTreeItem.h"
+#include "WorldTreeItem.h"
+
+FContentBundleHiearchy::FContentBundleHiearchy(FContentBundleMode* Mode)
+	:ISceneOutlinerHierarchy(Mode)
+{
+	UContentBundleEditorSubsystem::Get()->OnContentBundleChanged().AddRaw(this, &FContentBundleHiearchy::OnContentBundleChanged);
+	UContentBundleEditorSubsystem::Get()->OnContentBundleAdded().AddRaw(this, &FContentBundleHiearchy::OnContentBundleAdded);
+	UContentBundleEditorSubsystem::Get()->OnContentBundleRemoved().AddRaw(this, &FContentBundleHiearchy::OnContentBundleRemoved);
+}
+
+FContentBundleHiearchy::~FContentBundleHiearchy()
+{
+	UContentBundleEditorSubsystem::Get()->OnContentBundleChanged().RemoveAll(this);
+	UContentBundleEditorSubsystem::Get()->OnContentBundleAdded().RemoveAll(this);
+	UContentBundleEditorSubsystem::Get()->OnContentBundleRemoved().RemoveAll(this);
+}
+
+TUniquePtr<FContentBundleHiearchy> FContentBundleHiearchy::Create(FContentBundleMode* Mode)
+{
+	return TUniquePtr<FContentBundleHiearchy>(new FContentBundleHiearchy(Mode));
+}
+
+void FContentBundleHiearchy::CreateItems(TArray<FSceneOutlinerTreeItemPtr>& OutItems) const
+{
+	UContentBundleEditorSubsystem* ContentBundleEditorSubsystem = UContentBundleEditorSubsystem::Get();
+	if (ContentBundleEditorSubsystem != nullptr)
+	{
+		for (TSharedPtr<FContentBundleEditor>& ContentBundleEditor : ContentBundleEditorSubsystem->GetEditorContentBundles())
+		{
+			if (FSceneOutlinerTreeItemPtr ContentBundleTreeItem = Mode->CreateItemFor<FContentBundleTreeItem>(ContentBundleEditor, true))
+			{
+				OutItems.Add(ContentBundleTreeItem);
+			}
+		}
+	}	
+}
+
+void FContentBundleHiearchy::CreateChildren(const FSceneOutlinerTreeItemPtr& Item, TArray<FSceneOutlinerTreeItemPtr>& OutChildren) const
+{
+
+}
+
+FSceneOutlinerTreeItemPtr FContentBundleHiearchy::FindOrCreateParentItem(const ISceneOutlinerTreeItem& Item, const TMap<FSceneOutlinerTreeItemID, FSceneOutlinerTreeItemPtr>& Items, bool bCreate /* = false */)
+{
+	if (const FContentBundleTreeItem* ContentBundleTreeItem = Item.CastTo<FContentBundleTreeItem>())
+	{
+		TSharedPtr<FContentBundleEditor> ContentBundleEditorPin = ContentBundleTreeItem->GetContentBundleEditorPin();
+		if (ContentBundleEditorPin != nullptr)
+		{
+			if (UWorld* World = ContentBundleEditorPin->GetInjectedWorld())
+			{
+				if (const FSceneOutlinerTreeItemPtr* ParentItem = Items.Find(World))
+				{
+					return *ParentItem;
+				}
+				else if (bCreate)
+				{
+					return Mode->CreateItemFor<FWorldTreeItem>(World, true);
+				}
+			}
+		}
+	}
+
+	return nullptr;
+}
+
+void FContentBundleHiearchy::OnContentBundleChanged(const FContentBundleEditor* ContentBundle)
+{
+	FSceneOutlinerHierarchyChangedData EventData;
+	EventData.Type = FSceneOutlinerHierarchyChangedData::FullRefresh;
+	EventData.ItemIDs.Emplace(ContentBundle->GetGuid());
+	HierarchyChangedEvent.Broadcast(EventData);
+}
+
+void FContentBundleHiearchy::OnContentBundleAdded(const FContentBundleEditor* ContentBundle)
+{
+	FSceneOutlinerHierarchyChangedData EventData;
+	EventData.Type = FSceneOutlinerHierarchyChangedData::FullRefresh;
+	HierarchyChangedEvent.Broadcast(EventData);
+}
+
+void FContentBundleHiearchy::OnContentBundleRemoved(const FContentBundleEditor* ContentBundle)
+{
+	FSceneOutlinerHierarchyChangedData EventData;
+	EventData.Type = FSceneOutlinerHierarchyChangedData::FullRefresh;
+	HierarchyChangedEvent.Broadcast(EventData);
+}
\ No newline at end of file
diff --git a/Engine/Source/Editor/WorldPartitionEditor/Private/WorldPartition/ContentBundle/Outliner/ContentBundleMode.cpp b/Engine/Source/Editor/WorldPartitionEditor/Private/WorldPartition/ContentBundle/Outliner/ContentBundleMode.cpp
new file mode 100644
--- /dev/null
+++ b/Engine/Source/Editor/WorldPartitionEditor/Private/WorldPartition/ContentBundle/Outliner/ContentBundleMode.cpp
@@ -0,0 +1,305 @@
+// Copyright Epic Games, Inc. All Rights Reserved.
+
+#include "WorldPartition/ContentBundle/Outliner/ContentBundleMode.h"
+
+#include "WorldPartition/ContentBundle/Outliner/ContentBundleHierarchy.h"
+#include "ToolMenus.h"
+#include "SceneOutlinerMenuContext.h"
+#include "WorldPartition/ContentBundle/ContentBundleEditorSubsystem.h"
+#include "WorldPartition/ContentBundle/ContentBundleEditor.h"
+#include "WorldPartition/ContentBundle/Outliner/ContentBundleTreeItem.h"
+#include "Algo/Transform.h"
+
+#define LOCTEXT_NAMESPACE "ContentBundle"
+
+namespace ContentBundleOutlinerPrivate
+{
+	const FName DefaultContextBaseMenuName("ContentBundleOutliner.DefaultContextMenuBase");
+	const FName DefaultContextMenuName("ContentBundleOutliner.DefaultContextMenu");
+}
+
+FContentBundleMode::FContentBundleMode(const FContentBundleModeCreationParams& Params)
+	:ISceneOutlinerMode(Params.SceneOutliner)
+{
+
+}
+
+FContentBundleMode::~FContentBundleMode()
+{
+	UnregisterContextMenu();
+}
+
+void FContentBundleMode::Rebuild()
+{
+	Hierarchy = CreateHierarchy();
+}
+
+TSharedPtr<SWidget> FContentBundleMode::CreateContextMenu()
+{
+	RegisterContextMenu();
+
+	FSceneOutlinerItemSelection ItemSelection(SceneOutliner->GetSelection());
+
+	USceneOutlinerMenuContext* ContextObject = NewObject<USceneOutlinerMenuContext>();
+	ContextObject->SceneOutliner = StaticCastSharedRef<SSceneOutliner>(SceneOutliner->AsShared());
+	ContextObject->bShowParentTree = SceneOutliner->GetSharedData().bShowParentTree;
+	ContextObject->NumSelectedItems = ItemSelection.Num();
+	FToolMenuContext Context(ContextObject);
+
+	FName MenuName = ContentBundleOutlinerPrivate::DefaultContextMenuName;
+	SceneOutliner->GetSharedData().ModifyContextMenu.ExecuteIfBound(MenuName, Context);
+
+	// Build up the menu for a selection
+	UToolMenus* ToolMenus = UToolMenus::Get();
+	UToolMenu* Menu = ToolMenus->GenerateMenu(MenuName, Context);
+	for (const FToolMenuSection& Section : Menu->Sections)
+	{
+		if (Section.Blocks.Num() > 0)
+		{
+			return ToolMenus->GenerateWidget(Menu);
+		}
+	}
+
+	return nullptr;
+}
+
+TUniquePtr<ISceneOutlinerHierarchy> FContentBundleMode::CreateHierarchy()
+{
+	TUniquePtr<FContentBundleHiearchy> ContentBundleHierarchy = FContentBundleHiearchy::Create(this);
+	return ContentBundleHierarchy;
+}
+
+void FContentBundleMode::OnItemDoubleClick(FSceneOutlinerTreeItemPtr Item)
+{
+	if (FContentBundleTreeItem* ContentBundleTreeItem = Item->CastTo<FContentBundleTreeItem>())
+	{
+		ToggleContentBundleActivation(ContentBundleTreeItem->GetContentBundleEditor());
+	}
+}
+
+FReply FContentBundleMode::OnKeyDown(const FKeyEvent& InKeyEvent)
+{
+	if (InKeyEvent.GetKey() == EKeys::F5)
+	{
+		SceneOutliner->FullRefresh();
+		return FReply::Handled();
+	}
+
+	return FReply::Unhandled();
+}
+
+UWorld* FContentBundleMode::GetEditingWorld() const
+{
+	if (UContentBundleEditorSubsystem* ContentBundleEditorSubsystem = UContentBundleEditorSubsystem::Get())
+	{
+		return ContentBundleEditorSubsystem->GetWorld();
+	}
+
+	return nullptr;
+}
+
+void FContentBundleMode::RegisterContextMenu()
+{
+	UToolMenus* ToolMenus = UToolMenus::Get();
+
+	if (!ToolMenus->IsMenuRegistered(ContentBundleOutlinerPrivate::DefaultContextBaseMenuName))
+	{
+		UToolMenu* Menu = ToolMenus->RegisterMenu(ContentBundleOutlinerPrivate::DefaultContextBaseMenuName);
+
+		Menu->AddDynamicSection("ContentBundleDynamicSection", FNewToolMenuDelegate::CreateLambda([this](UToolMenu* InMenu)
+			{
+				USceneOutlinerMenuContext* Context = InMenu->FindContext<USceneOutlinerMenuContext>();
+				if (!Context || !Context->SceneOutliner.IsValid())
+				{
+					return;
+				}
+
+				SSceneOutliner* SceneOutliner = Context->SceneOutliner.Pin().Get();
+				
+				{
+					FToolMenuSection& Section = InMenu->AddSection("ContentBundleSection", LOCTEXT("ContentBundle", "Content Bundle"));
+
+					Section.AddMenuEntry("EditContentBundleMenuEntry", LOCTEXT("EditContentBundle", "Edit Content Bundle"), FText(), FSlateIcon(),
+						FUIAction(
+							FExecuteAction::CreateLambda([this, SceneOutliner]
+								{
+									TWeakPtr<FContentBundleEditor> SelectedContentBundleEditor = GetSelectedContentBundle(SceneOutliner);
+									ToggleContentBundleActivation(SelectedContentBundleEditor);
+								}),
+							FCanExecuteAction(),
+							FIsActionChecked::CreateLambda([this, SceneOutliner]
+								{
+									TSharedPtr<FContentBundleEditor> SelectedContentBundleEditor = GetSelectedContentBundlePin(SceneOutliner);
+									return SelectedContentBundleEditor != nullptr && SelectedContentBundleEditor->IsBeingEdited();
+								})
+							));
+
+					Section.AddMenuEntry("InjectBaseContentMenuEntry", LOCTEXT("InjectBaseContent", "Inject Base Content"), FText(), FSlateIcon(),
+						FUIAction(
+							FExecuteAction::CreateLambda([this, SceneOutliner]
+								{
+									if (TSharedPtr<FContentBundleEditor> SelectedContentBundleEditor = GetSelectedContentBundlePin(SceneOutliner))
+									{
+										SelectedContentBundleEditor->InjectBaseContent();
+									}
+								}),
+							FCanExecuteAction::CreateLambda([this, SceneOutliner]
+								{
+									if (TSharedPtr<FContentBundleEditor> SelectedContentBundleEditor = GetSelectedContentBundlePin(SceneOutliner))
+									{
+										return SelectedContentBundleEditor->IsBeingEdited() && SelectedContentBundleEditor->GetStatus() == EContentBundleStatus::ReadyToInject;
+									}
+									return false;
+								})
+							));
+
+					Section.AddMenuEntry("RemoveBaseContentMenuEntry", LOCTEXT("RemoveBaseContent", "Remove Base Content"), FText(), FSlateIcon(),
+						FUIAction(
+							FExecuteAction::CreateLambda([this, SceneOutliner]
+								{
+									if (TSharedPtr<FContentBundleEditor> SelectedContentBundleEditor = GetSelectedContentBundlePin(SceneOutliner))
+									{
+										SelectedContentBundleEditor->RemoveBaseContent();
+									}
+								}),
+							FCanExecuteAction::CreateLambda([this, SceneOutliner]
+								{
+									if (TSharedPtr<FContentBundleEditor> SelectedContentBundleEditor = GetSelectedContentBundlePin(SceneOutliner))
+									{
+										return SelectedContentBundleEditor->IsBeingEdited() && SelectedContentBundleEditor->GetStatus() == EContentBundleStatus::ContentInjected && !SelectedContentBundleEditor->HasUserPlacedActors();
+									}
+									return false;
+								})
+									));
+
+					
+				}
+
+				{
+					FToolMenuSection& Section = InMenu->AddSection("ContentBundleLoadingSection", LOCTEXT("Loading", "Loading"));
+
+					Section.AddMenuEntry("LoadActorsEntry", LOCTEXT("LoadActors", "Load Actors"), FText(), FSlateIcon(),
+						FUIAction(
+							FExecuteAction::CreateLambda([this, SceneOutliner]
+							{
+								TSharedPtr<FContentBundleEditor> SelectedContentBundleEditor = GetSelectedContentBundlePin(SceneOutliner);
+								if (SelectedContentBundleEditor != nullptr)
+								{
+									UContentBundleEditorSubsystem::Get()->ReferenceAllActors(*SelectedContentBundleEditor);
+								}
+							}),
+							FCanExecuteAction::CreateLambda([this]
+							{
+								return true;
+							})
+						));
+
+					Section.AddMenuEntry("UnloadActorsEntry", LOCTEXT("UnloadActors", "Unload Actors"), FText(), FSlateIcon(),
+						FUIAction(
+							FExecuteAction::CreateLambda([this, SceneOutliner]
+							{
+								TSharedPtr<FContentBundleEditor> SelectedContentBundleEditor = GetSelectedContentBundlePin(SceneOutliner);
+								if (SelectedContentBundleEditor != nullptr)
+								{
+									UContentBundleEditorSubsystem::Get()->UnreferenceAllActors(*SelectedContentBundleEditor);
+								}
+							}),
+							FCanExecuteAction::CreateLambda([this]
+							{
+								return true;
+							})
+						));
+				}
+
+				{
+					FToolMenuSection& Section = InMenu->AddSection("ContentBundleSelectionSection", LOCTEXT("Selection", "Selection"));
+
+					Section.AddMenuEntry("SelectActorsEntry", LOCTEXT("SelectActors", "Select Actors"), FText(), FSlateIcon(),
+						FUIAction(
+							FExecuteAction::CreateLambda([this, SceneOutliner]
+								{
+									TSharedPtr<FContentBundleEditor> SelectedContentBundleEditor = GetSelectedContentBundlePin(SceneOutliner);
+									if (SelectedContentBundleEditor != nullptr)
+									{
+										UContentBundleEditorSubsystem::Get()->SelectActors(*SelectedContentBundleEditor);
+									}
+								}),
+							FCanExecuteAction::CreateLambda([this]
+								{
+									return true;
+								})
+							));
+
+					Section.AddMenuEntry("DeselectActorsEntry", LOCTEXT("DeselectActors", "Deselect Actors"), FText(), FSlateIcon(),
+						FUIAction(
+							FExecuteAction::CreateLambda([this, SceneOutliner]
+								{
+									TSharedPtr<FContentBundleEditor> SelectedContentBundleEditor = GetSelectedContentBundlePin(SceneOutliner);
+									if (SelectedContentBundleEditor != nullptr)
+									{
+										UContentBundleEditorSubsystem::Get()->DeselectActors(*SelectedContentBundleEditor);
+									}
+								}),
+							FCanExecuteAction::CreateLambda([this]
+								{
+									return true;
+								})
+							));
+				}
+			}));
+	}
+
+	if (!ToolMenus->IsMenuRegistered(ContentBundleOutlinerPrivate::DefaultContextMenuName))
+	{
+		ToolMenus->RegisterMenu(ContentBundleOutlinerPrivate::DefaultContextMenuName, ContentBundleOutlinerPrivate::DefaultContextBaseMenuName);
+	}
+}
+
+void FContentBundleMode::UnregisterContextMenu()
+{
+	UToolMenus* ToolMenus = UToolMenus::Get();
+	ToolMenus->RemoveMenu(ContentBundleOutlinerPrivate::DefaultContextBaseMenuName);
+	ToolMenus->RemoveMenu(ContentBundleOutlinerPrivate::DefaultContextMenuName);
+}
+
+void FContentBundleMode::ToggleContentBundleActivation(const TWeakPtr<FContentBundleEditor>& ContentBundle)
+{
+	if (UContentBundleEditorSubsystem* ContentBundleEditorSubsystem = UContentBundleEditorSubsystem::Get())
+	{
+		if (UContentBundleEditionSubmodule* ContentBundleEditionModule = ContentBundleEditorSubsystem->GetEditionSubmodule())
+		{
+			if (TSharedPtr<FContentBundleEditor> ContentBundlePin = ContentBundle.Pin())
+			{
+				if (ContentBundlePin->IsBeingEdited())
+				{
+					ContentBundleEditionModule->DeactivateCurrentContentBundleEditing();
+				}
+				else
+				{
+					ContentBundleEditionModule->ActivateContentBundleEditing(ContentBundlePin);
+				}
+			}
+		}
+	}
+}
+
+TWeakPtr<FContentBundleEditor> FContentBundleMode::GetSelectedContentBundle(SSceneOutliner* InSceneOutliner) const
+{
+	FSceneOutlinerItemSelection ItemSelection(InSceneOutliner->GetSelection());
+	
+	TArray<FContentBundleTreeItem*> SelectedContentBundles;
+	ItemSelection.Get<FContentBundleTreeItem>(SelectedContentBundles);
+	
+	TArray<TWeakPtr<FContentBundleEditor>> ValidSelectedContentBundle;
+	Algo::TransformIf(SelectedContentBundles, ValidSelectedContentBundle, [](const auto Item) { return Item && Item->IsValid(); }, [](const auto Item) { return Item->GetContentBundleEditor(); });
+
+	check(ValidSelectedContentBundle.Num() == 1 || ValidSelectedContentBundle.IsEmpty());
+	return ValidSelectedContentBundle.IsEmpty() ? nullptr : ValidSelectedContentBundle[0];
+}
+
+TSharedPtr<FContentBundleEditor> FContentBundleMode::GetSelectedContentBundlePin(SSceneOutliner* InSceneOutliner) const
+{
+	return GetSelectedContentBundle(InSceneOutliner).Pin();
+}
+
+#undef LOCTEXT_NAMESPACE
\ No newline at end of file
diff --git a/Engine/Source/Editor/WorldPartitionEditor/Private/WorldPartition/ContentBundle/Outliner/ContentBundleStatusColumn.cpp b/Engine/Source/Editor/WorldPartitionEditor/Private/WorldPartition/ContentBundle/Outliner/ContentBundleStatusColumn.cpp
new file mode 100644
--- /dev/null
+++ b/Engine/Source/Editor/WorldPartitionEditor/Private/WorldPartition/ContentBundle/Outliner/ContentBundleStatusColumn.cpp
@@ -0,0 +1,73 @@
+// Copyright Epic Games, Inc. All Rights Reserved.
+
+#include "WorldPartition/ContentBundle/Outliner/ContentBundleStatusColumn.h"
+
+#include "Widgets/Images/SImage.h"
+#include "Widgets/Input/SButton.h"
+#include "Widgets/Input/SCheckBox.h"
+#include "Widgets/Views/STreeView.h"
+#include "ScopedTransaction.h"
+#include "Editor.h"
+#include "Styling/AppStyle.h"
+#include "WorldPartition/ContentBundle/Outliner/ContentBundleTreeItem.h"
+#include "WorldPartition/ContentBundle/ContentBundleEditor.h"
+
+#define LOCTEXT_NAMESPACE "ContentBundle"
+
+namespace ContentBundleOutlinerPrivate
+{
+	FName ContentBundleOutlinerStatus("Content Bundle Status");
+}
+
+FName FContentBundleOutlinerStatusColumn::GetID()
+{
+	return ContentBundleOutlinerPrivate::ContentBundleOutlinerStatus;
+}
+
+SHeaderRow::FColumn::FArguments FContentBundleOutlinerStatusColumn::ConstructHeaderRowColumn()
+{
+	return SHeaderRow::Column(GetColumnID())
+		.HAlignHeader(HAlign_Center)
+		.VAlignHeader(VAlign_Center)
+		.HAlignCell(HAlign_Center)
+		.VAlignCell(VAlign_Center)
+		.DefaultTooltip(FText::FromName(GetColumnID()))
+		[
+			SNew(STextBlock)
+			.Margin(FMargin(0.0f))
+			.Text(LOCTEXT("ContentBundleOutlinerStatus", "Status"))
+		];
+}
+
+const TSharedRef<SWidget> FContentBundleOutlinerStatusColumn::ConstructRowWidget(FSceneOutlinerTreeItemRef TreeItem, const STableRow<FSceneOutlinerTreeItemPtr>& Row)
+{
+	if (FContentBundleTreeItem* ContentBundleTreeItem = TreeItem->CastTo<FContentBundleTreeItem>())
+	{
+		TWeakPtr<FContentBundleEditor> ContentBundleEditor = ContentBundleTreeItem->GetContentBundleEditor();
+		if (ContentBundleEditor != nullptr)
+		{
+			return SNew(SHorizontalBox)
+				+ SHorizontalBox::Slot()
+				.Padding(0, 0, 0, 0)
+				.AutoWidth()
+				.VAlign(VAlign_Center)
+				[
+					SNew(STextBlock)
+					.Text_Lambda([ContentBundleTreeItem]
+						{ 
+							TSharedPtr<FContentBundleEditor> ContentBundleEditorPin = ContentBundleTreeItem->GetContentBundleEditorPin();
+							if (ContentBundleEditorPin != nullptr)
+							{
+								return UEnum::GetDisplayValueAsText(ContentBundleEditorPin->GetStatus());
+							}
+
+							return UEnum::GetDisplayValueAsText(EContentBundleStatus::Unknown);
+						})
+					.ColorAndOpacity(MakeAttributeLambda([ContentBundleTreeItem] { return ContentBundleTreeItem->GetItemColor(); }))
+				];
+		}
+	}
+	return SNullWidget::NullWidget;
+}
+
+#undef LOCTEXT_NAMESPACE
\ No newline at end of file
diff --git a/Engine/Source/Editor/WorldPartitionEditor/Private/WorldPartition/ContentBundle/Outliner/ContentBundleTreeItem.cpp b/Engine/Source/Editor/WorldPartitionEditor/Private/WorldPartition/ContentBundle/Outliner/ContentBundleTreeItem.cpp
new file mode 100644
--- /dev/null
+++ b/Engine/Source/Editor/WorldPartitionEditor/Private/WorldPartition/ContentBundle/Outliner/ContentBundleTreeItem.cpp
@@ -0,0 +1,97 @@
+// Copyright Epic Games, Inc. All Rights Reserved.
+
+#include "WorldPartition/ContentBundle/Outliner/ContentBundleTreeItem.h"
+
+#include "WorldPartition/ContentBundle/ContentBundle.h"
+#include "WorldPartition/ContentBundle/ContentBundleEditor.h"
+#include "WorldPartition/ContentBundle/ContentBundleEditorSubsystem.h"
+#include "Widgets/Text/STextBlock.h"
+
+#define LOCTEXT_NAMESPACE "ContentBundle"
+
+const FSceneOutlinerTreeItemType FContentBundleTreeItem::Type(&ISceneOutlinerTreeItem::Type);
+
+FContentBundleTreeItem::FContentBundleTreeItem(const TWeakPtr<FContentBundleEditor>& InContentBundleEditor)
+	: ISceneOutlinerTreeItem(Type)
+	, ContentBundleEditor(InContentBundleEditor)
+{
+	
+}
+
+bool FContentBundleTreeItem::IsValid() const
+{
+	return ContentBundleEditor.IsValid();
+}
+
+FSceneOutlinerTreeItemID FContentBundleTreeItem::GetID() const
+{
+	TSharedPtr<FContentBundleEditor> ContentBundleEditorPin = ContentBundleEditor.Pin();
+	if (ContentBundleEditorPin != nullptr)
+	{
+		return FSceneOutlinerTreeItemID(ContentBundleEditorPin->GetGuid());
+	}
+
+	return FSceneOutlinerTreeItemID();
+}
+
+FString FContentBundleTreeItem::GetDisplayString() const
+{
+	TSharedPtr<FContentBundleEditor> ContentBundleEditorPin = ContentBundleEditor.Pin();
+	if (ContentBundleEditorPin != nullptr)
+	{
+		return ContentBundleEditorPin->GetDisplayName();
+	}
+
+	return LOCTEXT("ContentBundleMissingForTreeItem", "(Invalid Content Bundle)").ToString();
+}
+
+bool FContentBundleTreeItem::CanInteract() const
+{
+	TSharedPtr<FContentBundleEditor> ContentBundleEditorPin = ContentBundleEditor.Pin();
+	if (ContentBundleEditorPin != nullptr)
+	{
+		if (UContentBundleEditorSubsystem* ContentBundleEditorSubsystem = UContentBundleEditorSubsystem::Get())
+		{
+			return ContentBundleEditorPin->GetInjectedWorld() == ContentBundleEditorSubsystem->GetWorld();
+		}
+	}
+	
+	return false;
+}
+
+TSharedRef<SWidget> FContentBundleTreeItem::GenerateLabelWidget(ISceneOutliner& Outliner, const STableRow<FSceneOutlinerTreeItemPtr>& InRow)
+{
+	
+	return SNew(SHorizontalBox)
+		+ SHorizontalBox::Slot()
+		.VAlign(VAlign_Center)
+		[
+			SNew(STextBlock)
+			.Text(MakeAttributeLambda([this] { return  FText::FromString(GetDisplayString()); }))
+			.ColorAndOpacity(MakeAttributeLambda([this] 
+				{ 
+					return GetItemColor();
+				}))
+		];
+}
+
+FSlateColor FContentBundleTreeItem::GetItemColor() const
+{
+	TSharedPtr<FContentBundleEditor> ContentBundleEditorPin = ContentBundleEditor.Pin();
+	if (ContentBundleEditorPin != nullptr)
+	{
+		if (!CanInteract())
+		{
+			return FSlateColor(FSceneOutlinerCommonLabelData::DarkColor);
+		}
+
+		if (ContentBundleEditorPin->IsBeingEdited())
+		{
+			return FAppStyle::Get().GetSlateColor("Colors.AccentGreen");
+		}
+	}
+
+	return FSlateColor::UseForeground();
+}
+
+#undef LOCTEXT_NAMESPACE
\ No newline at end of file
diff --git a/Engine/Source/Editor/WorldPartitionEditor/Private/WorldPartition/ContentBundle/Outliner/SContentBundleOutliner.cpp b/Engine/Source/Editor/WorldPartitionEditor/Private/WorldPartition/ContentBundle/Outliner/SContentBundleOutliner.cpp
new file mode 100644
--- /dev/null
+++ b/Engine/Source/Editor/WorldPartitionEditor/Private/WorldPartition/ContentBundle/Outliner/SContentBundleOutliner.cpp
@@ -0,0 +1,18 @@
+// Copyright Epic Games, Inc. All Rights Reserved.
+
+#include "WorldPartition/ContentBundle/Outliner/SContentBundleOutliner.h"
+
+#include "WorldPartition/ContentBundle/ContentBundleEditor.h"
+
+void SContentBundleOutliner::SelectContentBundle(const TWeakPtr<FContentBundleEditor>& ContentBundle)
+{
+	TSharedPtr<FContentBundleEditor> ContentBundleEditorPin = ContentBundle.Pin();
+	if (ContentBundleEditorPin != nullptr)
+	{
+		FSceneOutlinerTreeItemPtr Item = GetTreeItem(ContentBundleEditorPin->GetGuid());
+		if (Item->IsValid())
+		{
+			SetItemSelection(Item, true, ESelectInfo::Direct);
+		}
+	}	
+}
\ No newline at end of file
diff --git a/Engine/Source/Editor/WorldPartitionEditor/Private/WorldPartition/ContentBundle/SContentBundleBrowser.cpp b/Engine/Source/Editor/WorldPartitionEditor/Private/WorldPartition/ContentBundle/SContentBundleBrowser.cpp
new file mode 100644
--- /dev/null
+++ b/Engine/Source/Editor/WorldPartitionEditor/Private/WorldPartition/ContentBundle/SContentBundleBrowser.cpp
@@ -0,0 +1,63 @@
+// Copyright Epic Games, Inc. All Rights Reserved.
+
+#include "WorldPartition/ContentBundle/SContentBundleBrowser.h"
+
+#include "SceneOutlinerPublicTypes.h"
+#include "WorldPartition/ContentBundle/Outliner/SContentBundleOutliner.h"
+#include "WorldPartition/ContentBundle/Outliner/ContentBundleMode.h"
+#include "WorldPartition/ContentBundle/Outliner/ContentBundleEditingColumn.h"
+#include "WorldPartition/ContentBundle/Outliner/ContentBundleStatusColumn.h"
+#include "WorldPartition/ContentBundle/Outliner/ContentBundleClientColumn.h"
+#include "WorldPartition/ContentBundle/Outliner/ContentBundleHasErrorColumn.h"
+#include "WorldPartition/ContentBundle/Outliner/ContentBundleActorCountColumn.h"
+
+namespace ContentBundleOutlinerPrivate
+{
+	enum EColumnPriority
+	{
+		Client = 0,
+		BuiltIn,
+		Status,
+		Editing,
+		ActorCount,
+		HasError
+	};
+}
+
+
+void SContentBundleBrowser::Construct(const FArguments& InArgs)
+{
+	FSceneOutlinerInitializationOptions InitOptions;
+	InitOptions.bShowHeaderRow = true;
+	InitOptions.bShowParentTree = true;
+	InitOptions.bShowCreateNewFolder = false;
+	InitOptions.bShowTransient = true;
+	InitOptions.ModeFactory = FCreateSceneOutlinerMode::CreateLambda([this](SSceneOutliner* Outliner) { return new FContentBundleMode(FContentBundleModeCreationParams(Outliner)); });
+	InitOptions.ColumnMap.Add(FContentBundleOutlinerClientColumn::GetID(), FSceneOutlinerColumnInfo(ESceneOutlinerColumnVisibility::Visible, ContentBundleOutlinerPrivate::EColumnPriority::Client, FCreateSceneOutlinerColumn::CreateLambda([](ISceneOutliner& InSceneOutliner) { return MakeShareable(new FContentBundleOutlinerClientColumn(InSceneOutliner)); }), true));
+	InitOptions.ColumnMap.Add(FSceneOutlinerBuiltInColumnTypes::Label(), FSceneOutlinerColumnInfo(ESceneOutlinerColumnVisibility::Visible, ContentBundleOutlinerPrivate::EColumnPriority::BuiltIn));
+	InitOptions.ColumnMap.Add(FContentBundleOutlinerStatusColumn::GetID(), FSceneOutlinerColumnInfo(ESceneOutlinerColumnVisibility::Visible, ContentBundleOutlinerPrivate::EColumnPriority::Status, FCreateSceneOutlinerColumn::CreateLambda([](ISceneOutliner& InSceneOutliner) { return MakeShareable(new FContentBundleOutlinerStatusColumn(InSceneOutliner)); }), false));
+	InitOptions.ColumnMap.Add(FContentBundleOutlinerEditingColumn::GetID(), FSceneOutlinerColumnInfo(ESceneOutlinerColumnVisibility::Visible, ContentBundleOutlinerPrivate::EColumnPriority::Editing, FCreateSceneOutlinerColumn::CreateLambda([](ISceneOutliner& InSceneOutliner) { return MakeShareable(new FContentBundleOutlinerEditingColumn(InSceneOutliner)); }), false));
+	InitOptions.ColumnMap.Add(FContentBundleOutlinerActorCountColumn::GetID(), FSceneOutlinerColumnInfo(ESceneOutlinerColumnVisibility::Visible, ContentBundleOutlinerPrivate::EColumnPriority::ActorCount, FCreateSceneOutlinerColumn::CreateLambda([](ISceneOutliner& InSceneOutliner) { return MakeShareable(new FContentBundleOutlinerActorCountColumn(InSceneOutliner)); }), false));
+	InitOptions.ColumnMap.Add(FContentBundleOutlinerHasErrorColumn::GetID(), FSceneOutlinerColumnInfo(ESceneOutlinerColumnVisibility::Visible, ContentBundleOutlinerPrivate::EColumnPriority::HasError, FCreateSceneOutlinerColumn::CreateLambda([](ISceneOutliner& InSceneOutliner) { return MakeShareable(new FContentBundleOutlinerHasErrorColumn(InSceneOutliner)); }), false));
+	
+	ContentBundleOutliner = SNew(SContentBundleOutliner, InitOptions)
+		.IsEnabled(FSlateApplication::Get().GetNormalExecutionAttribute());
+
+	ChildSlot
+	[
+		SAssignNew(ContentAreaBox, SVerticalBox)
+		.IsEnabled_Lambda([]() { return GWorld ? UWorld::IsPartitionedWorld(GWorld) : false; })
+	];
+
+	ContentAreaBox->AddSlot()
+		.AutoHeight()
+		.FillHeight(1.0f)
+		[
+			ContentBundleOutliner.ToSharedRef()
+		];
+}
+
+void SContentBundleBrowser::SelectContentBundle(const TWeakPtr<FContentBundleEditor>& ContentBundle)
+{
+	ContentBundleOutliner->SelectContentBundle(ContentBundle);
+}
\ No newline at end of file
diff --git a/Engine/Source/Editor/WorldPartitionEditor/Private/WorldPartitionEditorModule.cpp b/Engine/Source/Editor/WorldPartitionEditor/Private/WorldPartitionEditorModule.cpp
--- a/Engine/Source/Editor/WorldPartitionEditor/Private/WorldPartitionEditorModule.cpp
+++ b/Engine/Source/Editor/WorldPartitionEditor/Private/WorldPartitionEditorModule.cpp
@@ -44,12 +44,18 @@
 #include "Widgets/Docking/SDockTab.h"
 
 #include "Styling/AppStyle.h"
+#include "WorldPartition/ContentBundle/SContentBundleBrowser.h"
+#include "Selection.h"
+#include "WorldPartition/ContentBundle/ContentBundleEditorSubsystem.h"
+#include "UObject/UObjectBase.h"
+#include "WorldPartition/ContentBundle/ContentBundleEditor.h"
 
 IMPLEMENT_MODULE( FWorldPartitionEditorModule, WorldPartitionEditor );
 
 #define LOCTEXT_NAMESPACE "WorldPartition"
 
 const FName WorldPartitionEditorTabId("WorldBrowserPartitionEditor");
+const FName ContentBundleBrowserTabId("ContentBundleBrowser");
 
 DEFINE_LOG_CATEGORY_STATIC(LogWorldPartitionEditor, All, All);
 
@@ -210,6 +216,14 @@ TSharedRef<SWidget> FWorldPartitionEditorModule::CreateWorldPartitionEditor()
 	return SNew(SWorldPartitionEditor).InWorld(EditorWorld);
 }
 
+TSharedRef<SWidget> FWorldPartitionEditorModule::CreateContentBundleBrowser()
+{
+	check(ContentBundleBrowser == nullptr);
+	TSharedRef<SContentBundleBrowser> NewDataLayerBrowser = SNew(SContentBundleBrowser);
+	ContentBundleBrowser = NewDataLayerBrowser;
+	return NewDataLayerBrowser;
+}
+
 int32 FWorldPartitionEditorModule::GetPlacementGridSize() const
 {
 	// Currently shares setting with Foliage. Can be changed when exposed.
@@ -612,6 +626,19 @@ TSharedRef<SDockTab> FWorldPartitionEditorModule::SpawnWorldPartitionTab(const F
 	return NewTab;
 }
 
+TSharedRef<SDockTab> FWorldPartitionEditorModule::SpawnContentBundleTab(const FSpawnTabArgs& Args)
+{
+	TSharedRef<SDockTab> NewTab =
+		SNew(SDockTab)
+		.Label(NSLOCTEXT("LevelEditor", "ContentBundleTabTitle", "Content Bundles"))
+		[
+			CreateContentBundleBrowser()
+		];
+
+	ContentBundleTab = NewTab;
+	return NewTab;
+}
+
 void FWorldPartitionEditorModule::RegisterWorldPartitionTabs(TSharedPtr<FTabManager> InTabManager)
 {
 	const IWorkspaceMenuStructure& MenuStructure = WorkspaceMenu::GetMenuStructure();
@@ -624,6 +651,15 @@ void FWorldPartitionEditorModule::RegisterWorldPartitionTabs(TSharedPtr<FTabMana
 		.SetTooltipText(NSLOCTEXT("LevelEditorTabs", "WorldPartitionEditorTooltipText", "Open the World Partition Editor."))
 		.SetGroup(MenuStructure.GetLevelEditorWorldPartitionCategory())
 		.SetIcon(WorldPartitionIcon);
+
+	constexpr TCHAR PLACEHOLDER_ContentBundleIcon[] = TEXT("LevelEditor.Tabs.DataLayers"); // todo_ow: create placeholder icon for content bundle tab
+	const FSlateIcon DataLayersIcon(FAppStyle::GetAppStyleSetName(), PLACEHOLDER_ContentBundleIcon);
+	InTabManager->RegisterTabSpawner(ContentBundleBrowserTabId,
+		FOnSpawnTab::CreateRaw(this, &FWorldPartitionEditorModule::SpawnContentBundleTab))
+		.SetDisplayName(NSLOCTEXT("LevelEditorTabs", "LevelEditorContentBundleBrowser", "Content Bundles Outliner"))
+		.SetTooltipText(NSLOCTEXT("LevelEditorTabs", "LevelEditorContentBundleBrowserTooltipText", "Open the Content Bundles Outliner."))
+		.SetGroup(MenuStructure.GetLevelEditorWorldPartitionCategory())
+		.SetIcon(DataLayersIcon);
 }
 
 void FWorldPartitionEditorModule::RegisterWorldPartitionLayout(FLayoutExtender& Extender)
diff --git a/Engine/Source/Editor/WorldPartitionEditor/Public/WorldPartition/ContentBundle/Outliner/ContentBundleActorCountColumn.h b/Engine/Source/Editor/WorldPartitionEditor/Public/WorldPartition/ContentBundle/Outliner/ContentBundleActorCountColumn.h
new file mode 100644
--- /dev/null
+++ b/Engine/Source/Editor/WorldPartitionEditor/Public/WorldPartition/ContentBundle/Outliner/ContentBundleActorCountColumn.h
@@ -0,0 +1,32 @@
+// Copyright Epic Games, Inc. All Rights Reserved.
+
+#pragma once
+
+#include "CoreMinimal.h"
+#include "SceneOutlinerFwd.h"
+#include "ISceneOutliner.h"
+#include "ISceneOutlinerColumn.h"
+#include "Widgets/Views/SHeaderRow.h"
+#include "SceneOutlinerGutter.h"
+
+template<typename ItemType> class STableRow;
+class FContentBundleOutlinerActorCountColumn : public ISceneOutlinerColumn
+{
+public:
+	FContentBundleOutlinerActorCountColumn(ISceneOutliner& SceneOutliner) : WeakSceneOutliner(StaticCastSharedRef<ISceneOutliner>(SceneOutliner.AsShared())) {}
+	virtual ~FContentBundleOutlinerActorCountColumn() {}
+	static FName GetID();
+
+	//////////////////////////////////////////////////////////////////////////
+	// Begin ISceneOutlinerColumn Implementation
+	virtual FName GetColumnID() override { return GetID(); }
+	virtual SHeaderRow::FColumn::FArguments ConstructHeaderRowColumn() override;
+	virtual const TSharedRef<SWidget> ConstructRowWidget(FSceneOutlinerTreeItemRef TreeItem, const STableRow<FSceneOutlinerTreeItemPtr>& Row) override;
+	virtual bool SupportsSorting() const override { return false; }
+	// End ISceneOutlinerColumn Implementation
+	//////////////////////////////////////////////////////////////////////////
+
+private:
+	TWeakPtr<ISceneOutliner> WeakSceneOutliner;
+};
+
diff --git a/Engine/Source/Editor/WorldPartitionEditor/Public/WorldPartition/ContentBundle/Outliner/ContentBundleClientColumn.h b/Engine/Source/Editor/WorldPartitionEditor/Public/WorldPartition/ContentBundle/Outliner/ContentBundleClientColumn.h
new file mode 100644
--- /dev/null
+++ b/Engine/Source/Editor/WorldPartitionEditor/Public/WorldPartition/ContentBundle/Outliner/ContentBundleClientColumn.h
@@ -0,0 +1,36 @@
+// Copyright Epic Games, Inc. All Rights Reserved.
+
+#pragma once
+
+#include "CoreMinimal.h"
+#include "SceneOutlinerFwd.h"
+#include "ISceneOutliner.h"
+#include "ISceneOutlinerColumn.h"
+#include "Widgets/Views/SHeaderRow.h"
+#include "SceneOutlinerGutter.h"
+
+class FContentBundleEditor;
+
+template<typename ItemType> class STableRow;
+class FContentBundleOutlinerClientColumn : public ISceneOutlinerColumn
+{
+public:
+	FContentBundleOutlinerClientColumn(ISceneOutliner& SceneOutliner) : WeakSceneOutliner(StaticCastSharedRef<ISceneOutliner>(SceneOutliner.AsShared())) {}
+	virtual ~FContentBundleOutlinerClientColumn() {}
+	static FName GetID();
+
+	//////////////////////////////////////////////////////////////////////////
+	// Begin ISceneOutlinerColumn Implementation
+	virtual FName GetColumnID() override { return GetID(); }
+	virtual SHeaderRow::FColumn::FArguments ConstructHeaderRowColumn() override;
+	virtual const TSharedRef<SWidget> ConstructRowWidget(FSceneOutlinerTreeItemRef TreeItem, const STableRow<FSceneOutlinerTreeItemPtr>& Row) override;
+	virtual bool SupportsSorting() const override { return true; }
+	virtual void SortItems(TArray<FSceneOutlinerTreeItemPtr>& OutItems, const EColumnSortMode::Type SortMode) const override;
+	// End ISceneOutlinerColumn Implementation
+	//////////////////////////////////////////////////////////////////////////
+
+private:
+	FText GetClientDisplayName(const TWeakPtr<FContentBundleEditor>& ContentBundleEditor) const;
+
+	TWeakPtr<ISceneOutliner> WeakSceneOutliner;
+};
diff --git a/Engine/Source/Editor/WorldPartitionEditor/Public/WorldPartition/ContentBundle/Outliner/ContentBundleEditingColumn.h b/Engine/Source/Editor/WorldPartitionEditor/Public/WorldPartition/ContentBundle/Outliner/ContentBundleEditingColumn.h
new file mode 100644
--- /dev/null
+++ b/Engine/Source/Editor/WorldPartitionEditor/Public/WorldPartition/ContentBundle/Outliner/ContentBundleEditingColumn.h
@@ -0,0 +1,31 @@
+// Copyright Epic Games, Inc. All Rights Reserved.
+
+#pragma once
+
+#include "CoreMinimal.h"
+#include "SceneOutlinerFwd.h"
+#include "ISceneOutliner.h"
+#include "ISceneOutlinerColumn.h"
+#include "Widgets/Views/SHeaderRow.h"
+#include "SceneOutlinerGutter.h"
+
+template<typename ItemType> class STableRow;
+class FContentBundleOutlinerEditingColumn : public ISceneOutlinerColumn
+{
+public:
+	FContentBundleOutlinerEditingColumn(ISceneOutliner& SceneOutliner) : WeakSceneOutliner(StaticCastSharedRef<ISceneOutliner>(SceneOutliner.AsShared())) {}
+	virtual ~FContentBundleOutlinerEditingColumn() {}
+	static FName GetID();
+
+	//////////////////////////////////////////////////////////////////////////
+	// Begin ISceneOutlinerColumn Implementation
+	virtual FName GetColumnID() override { return GetID(); }
+	virtual SHeaderRow::FColumn::FArguments ConstructHeaderRowColumn() override;
+	virtual const TSharedRef<SWidget> ConstructRowWidget(FSceneOutlinerTreeItemRef TreeItem, const STableRow<FSceneOutlinerTreeItemPtr>& Row) override;
+	virtual bool SupportsSorting() const override { return false; }
+	// End ISceneOutlinerColumn Implementation
+	//////////////////////////////////////////////////////////////////////////
+
+private:
+	TWeakPtr<ISceneOutliner> WeakSceneOutliner;
+};
diff --git a/Engine/Source/Editor/WorldPartitionEditor/Public/WorldPartition/ContentBundle/Outliner/ContentBundleHasErrorColumn.h b/Engine/Source/Editor/WorldPartitionEditor/Public/WorldPartition/ContentBundle/Outliner/ContentBundleHasErrorColumn.h
new file mode 100644
--- /dev/null
+++ b/Engine/Source/Editor/WorldPartitionEditor/Public/WorldPartition/ContentBundle/Outliner/ContentBundleHasErrorColumn.h
@@ -0,0 +1,31 @@
+// Copyright Epic Games, Inc. All Rights Reserved.
+
+#pragma once
+
+#include "CoreMinimal.h"
+#include "SceneOutlinerFwd.h"
+#include "ISceneOutliner.h"
+#include "ISceneOutlinerColumn.h"
+#include "Widgets/Views/SHeaderRow.h"
+#include "SceneOutlinerGutter.h"
+
+template<typename ItemType> class STableRow;
+class FContentBundleOutlinerHasErrorColumn : public ISceneOutlinerColumn
+{
+public:
+	FContentBundleOutlinerHasErrorColumn(ISceneOutliner& SceneOutliner) : WeakSceneOutliner(StaticCastSharedRef<ISceneOutliner>(SceneOutliner.AsShared())) {}
+	virtual ~FContentBundleOutlinerHasErrorColumn() {}
+	static FName GetID();
+
+	//////////////////////////////////////////////////////////////////////////
+	// Begin ISceneOutlinerColumn Implementation
+	virtual FName GetColumnID() override { return GetID(); }
+	virtual SHeaderRow::FColumn::FArguments ConstructHeaderRowColumn() override;
+	virtual const TSharedRef<SWidget> ConstructRowWidget(FSceneOutlinerTreeItemRef TreeItem, const STableRow<FSceneOutlinerTreeItemPtr>& Row) override;
+	virtual bool SupportsSorting() const override { return false; }
+	// End ISceneOutlinerColumn Implementation
+	//////////////////////////////////////////////////////////////////////////
+
+private:
+	TWeakPtr<ISceneOutliner> WeakSceneOutliner;
+};
diff --git a/Engine/Source/Editor/WorldPartitionEditor/Public/WorldPartition/ContentBundle/Outliner/ContentBundleHierarchy.h b/Engine/Source/Editor/WorldPartitionEditor/Public/WorldPartition/ContentBundle/Outliner/ContentBundleHierarchy.h
new file mode 100644
--- /dev/null
+++ b/Engine/Source/Editor/WorldPartitionEditor/Public/WorldPartition/ContentBundle/Outliner/ContentBundleHierarchy.h
@@ -0,0 +1,30 @@
+// Copyright Epic Games, Inc. All Rights Reserved.
+
+#pragma once
+
+#include "CoreMinimal.h"
+#include "ISceneOutlinerHierarchy.h"
+
+class FContentBundleMode;
+class FContentBundleEditor;
+
+class FContentBundleHiearchy : public ISceneOutlinerHierarchy
+{
+public:
+	static TUniquePtr<FContentBundleHiearchy> Create(FContentBundleMode* Mode);
+
+	virtual ~FContentBundleHiearchy();
+
+	//~ Begin ISceneOutlinerHierarchy interface
+	virtual void CreateItems(TArray<FSceneOutlinerTreeItemPtr>& OutItems) const override;
+	virtual void CreateChildren(const FSceneOutlinerTreeItemPtr& Item, TArray<FSceneOutlinerTreeItemPtr>& OutChildren) const override;
+	virtual FSceneOutlinerTreeItemPtr FindOrCreateParentItem(const ISceneOutlinerTreeItem& Item, const TMap<FSceneOutlinerTreeItemID, FSceneOutlinerTreeItemPtr>& Items, bool bCreate = false) override;
+	//~ End ISceneOutlinerHierarchy interface
+
+private:
+	void OnContentBundleChanged(const FContentBundleEditor* ContentBundle);
+	void OnContentBundleAdded(const FContentBundleEditor* ContentBundle);
+	void OnContentBundleRemoved(const FContentBundleEditor* ContentBundle);
+
+	FContentBundleHiearchy(FContentBundleMode* Mode);
+};
\ No newline at end of file
diff --git a/Engine/Source/Editor/WorldPartitionEditor/Public/WorldPartition/ContentBundle/Outliner/ContentBundleMode.h b/Engine/Source/Editor/WorldPartitionEditor/Public/WorldPartition/ContentBundle/Outliner/ContentBundleMode.h
new file mode 100644
--- /dev/null
+++ b/Engine/Source/Editor/WorldPartitionEditor/Public/WorldPartition/ContentBundle/Outliner/ContentBundleMode.h
@@ -0,0 +1,48 @@
+// Copyright Epic Games, Inc. All Rights Reserved.
+
+#pragma once
+
+#include "CoreMinimal.h"
+#include "ISceneOutlinerMode.h"
+
+class SSceneOutliner;
+class FContentBundleEditor;
+
+struct FContentBundleModeCreationParams
+{
+	FContentBundleModeCreationParams(SSceneOutliner* InSceneOutliner)
+		:SceneOutliner(InSceneOutliner)
+	{}
+
+	SSceneOutliner* SceneOutliner;
+};
+
+
+class FContentBundleMode : public ISceneOutlinerMode
+{
+public:
+	FContentBundleMode(const FContentBundleModeCreationParams& Params);
+	~FContentBundleMode();
+
+	//~ Begin ISceneOutlinerMode interface
+	virtual void Rebuild() override;
+	virtual TSharedPtr<SWidget> CreateContextMenu() override;
+
+	virtual void OnItemDoubleClick(FSceneOutlinerTreeItemPtr Item) override;
+	virtual FReply OnKeyDown(const FKeyEvent& InKeyEvent) override;
+	//~ Begin ISceneOutlinerMode interface
+
+	UWorld* GetEditingWorld() const;
+
+protected:
+	//~ Begin ISceneOutlinerMode interface
+	virtual TUniquePtr<ISceneOutlinerHierarchy> CreateHierarchy() override;
+	//~ End ISceneOutlinerMode interface
+
+	TWeakPtr<FContentBundleEditor> GetSelectedContentBundle(SSceneOutliner* InSceneOutliner) const;
+	TSharedPtr<FContentBundleEditor> GetSelectedContentBundlePin(SSceneOutliner* InSceneOutliner) const;
+	void ToggleContentBundleActivation(const TWeakPtr<FContentBundleEditor>& ContentBundle);
+
+	void RegisterContextMenu();
+	void UnregisterContextMenu();
+};
\ No newline at end of file
diff --git a/Engine/Source/Editor/WorldPartitionEditor/Public/WorldPartition/ContentBundle/Outliner/ContentBundleStatusColumn.h b/Engine/Source/Editor/WorldPartitionEditor/Public/WorldPartition/ContentBundle/Outliner/ContentBundleStatusColumn.h
new file mode 100644
--- /dev/null
+++ b/Engine/Source/Editor/WorldPartitionEditor/Public/WorldPartition/ContentBundle/Outliner/ContentBundleStatusColumn.h
@@ -0,0 +1,31 @@
+// Copyright Epic Games, Inc. All Rights Reserved.
+
+#pragma once
+
+#include "CoreMinimal.h"
+#include "SceneOutlinerFwd.h"
+#include "ISceneOutliner.h"
+#include "ISceneOutlinerColumn.h"
+#include "Widgets/Views/SHeaderRow.h"
+#include "SceneOutlinerGutter.h"
+
+template<typename ItemType> class STableRow;
+class FContentBundleOutlinerStatusColumn : public ISceneOutlinerColumn
+{
+public:
+	FContentBundleOutlinerStatusColumn(ISceneOutliner& SceneOutliner) : WeakSceneOutliner(StaticCastSharedRef<ISceneOutliner>(SceneOutliner.AsShared())) {}
+	virtual ~FContentBundleOutlinerStatusColumn() {}
+	static FName GetID();
+
+	//////////////////////////////////////////////////////////////////////////
+	// Begin ISceneOutlinerColumn Implementation
+	virtual FName GetColumnID() override { return GetID(); }
+	virtual SHeaderRow::FColumn::FArguments ConstructHeaderRowColumn() override;
+	virtual const TSharedRef<SWidget> ConstructRowWidget(FSceneOutlinerTreeItemRef TreeItem, const STableRow<FSceneOutlinerTreeItemPtr>& Row) override;
+	virtual bool SupportsSorting() const override { return false; }
+	// End ISceneOutlinerColumn Implementation
+	//////////////////////////////////////////////////////////////////////////
+
+private:
+	TWeakPtr<ISceneOutliner> WeakSceneOutliner;
+};
diff --git a/Engine/Source/Editor/WorldPartitionEditor/Public/WorldPartition/ContentBundle/Outliner/ContentBundleTreeItem.h b/Engine/Source/Editor/WorldPartitionEditor/Public/WorldPartition/ContentBundle/Outliner/ContentBundleTreeItem.h
new file mode 100644
--- /dev/null
+++ b/Engine/Source/Editor/WorldPartitionEditor/Public/WorldPartition/ContentBundle/Outliner/ContentBundleTreeItem.h
@@ -0,0 +1,46 @@
+// Copyright Epic Games, Inc. All Rights Reserved.
+
+#pragma once
+
+#include "CoreMinimal.h"
+#include "Misc/Attribute.h"
+#include "SceneOutlinerFwd.h"
+#include "ISceneOutlinerTreeItem.h"
+
+class FContentBundleEditor;
+
+struct FContentBundleTreeItem : ISceneOutlinerTreeItem
+{
+public:
+	FContentBundleTreeItem(const TWeakPtr<FContentBundleEditor>& InContentBundleEditor);
+	TWeakPtr<FContentBundleEditor> GetContentBundleEditor() const { return IsValid() ? ContentBundleEditor : nullptr; }
+	TSharedPtr<FContentBundleEditor> GetContentBundleEditorPin() const { return IsValid() ? ContentBundleEditor.Pin() : nullptr; }
+
+	/* Begin ISceneOutlinerTreeItem Implementation */
+	virtual bool IsValid() const override;
+	virtual FSceneOutlinerTreeItemID GetID() const override;
+	virtual FString GetDisplayString() const override;
+	virtual bool CanInteract() const override;
+	virtual TSharedRef<SWidget> GenerateLabelWidget(ISceneOutliner& Outliner, const STableRow<FSceneOutlinerTreeItemPtr>& InRow) override;
+	/* End ISceneOutlinerTreeItem Implementation */
+
+	static const FSceneOutlinerTreeItemType Type;
+
+	DECLARE_DELEGATE_RetVal_OneParam(bool, FFilterPredicate, const FContentBundleEditor*);
+	DECLARE_DELEGATE_RetVal_OneParam(bool, FInteractivePredicate, const FContentBundleEditor*);
+
+	bool Filter(FFilterPredicate Pred) const
+	{
+		return Pred.Execute(GetContentBundleEditorPin().Get());
+	}
+
+	bool GetInteractiveState(FInteractivePredicate Pred) const
+	{
+		return Pred.Execute(GetContentBundleEditorPin().Get());
+	}
+
+	FSlateColor GetItemColor() const;
+
+private:
+	TWeakPtr<FContentBundleEditor> ContentBundleEditor;
+};
\ No newline at end of file
diff --git a/Engine/Source/Editor/WorldPartitionEditor/Public/WorldPartition/ContentBundle/Outliner/SContentBundleOutliner.h b/Engine/Source/Editor/WorldPartitionEditor/Public/WorldPartition/ContentBundle/Outliner/SContentBundleOutliner.h
new file mode 100644
--- /dev/null
+++ b/Engine/Source/Editor/WorldPartitionEditor/Public/WorldPartition/ContentBundle/Outliner/SContentBundleOutliner.h
@@ -0,0 +1,22 @@
+// Copyright Epic Games, Inc. All Rights Reserved.
+
+#pragma once
+
+#include "CoreMinimal.h"
+#include "SSceneOutliner.h"
+
+class FContentBundleEditor;
+
+class SContentBundleOutliner : public SSceneOutliner
+{
+public:
+	void Construct(const FArguments& InArgs, const FSceneOutlinerInitializationOptions& InitOptions)
+	{
+		SSceneOutliner::Construct(InArgs, InitOptions);
+	}
+
+	SContentBundleOutliner() {}
+	virtual ~SContentBundleOutliner() {}
+
+	void SelectContentBundle(const TWeakPtr<FContentBundleEditor>& ContentBundle);
+};
\ No newline at end of file
diff --git a/Engine/Source/Editor/WorldPartitionEditor/Public/WorldPartition/ContentBundle/SContentBundleBrowser.h b/Engine/Source/Editor/WorldPartitionEditor/Public/WorldPartition/ContentBundle/SContentBundleBrowser.h
new file mode 100644
--- /dev/null
+++ b/Engine/Source/Editor/WorldPartitionEditor/Public/WorldPartition/ContentBundle/SContentBundleBrowser.h
@@ -0,0 +1,25 @@
+// Copyright Epic Games, Inc. All Rights Reserved.
+#pragma once
+
+#include "CoreMinimal.h"
+#include "Widgets/SCompoundWidget.h"
+
+class SContentBundleOutliner;
+class FContentBundleEditor;
+class SVerticalBox;
+
+class SContentBundleBrowser : public SCompoundWidget
+{
+public:
+
+	SLATE_BEGIN_ARGS(SContentBundleBrowser) {}
+	SLATE_END_ARGS()
+
+	void Construct(const FArguments& InArgs);
+
+	void SelectContentBundle(const TWeakPtr<FContentBundleEditor>& ContentBundle);
+
+private:
+	TSharedPtr<SContentBundleOutliner> ContentBundleOutliner;
+	TSharedPtr<SVerticalBox> ContentAreaBox;
+};
diff --git a/Engine/Source/Editor/WorldPartitionEditor/Public/WorldPartitionEditorModule.h b/Engine/Source/Editor/WorldPartitionEditor/Public/WorldPartitionEditorModule.h
--- a/Engine/Source/Editor/WorldPartitionEditor/Public/WorldPartitionEditorModule.h
+++ b/Engine/Source/Editor/WorldPartitionEditor/Public/WorldPartitionEditorModule.h
@@ -8,6 +8,7 @@ class FTabManager;
 class FLayoutExtender;
 class SDockTab;
 class FSpawnTabArgs;
+class SContentBundleBrowser;
 
 
 /**
@@ -62,6 +63,10 @@ public:
 	/** Return the commandlet execution event */
 	virtual FOnExecuteCommandlet& OnExecuteCommandlet() override { return OnExecuteCommandletEvent; }
 
+	/**
+	 * Creates a Content Bundle Browser widget
+	 */
+	TSharedRef<class SWidget> CreateContentBundleBrowser();
 private:
 	/** Called when the level editors map changes. We will determine if the new map is a valid world partition world and close world partition tabs if not */
 	void OnMapChanged(uint32 MapFlags);
@@ -78,6 +83,9 @@ private:
 	/** Spawns the world partition tab */
 	TSharedRef<SDockTab> SpawnWorldPartitionTab(const FSpawnTabArgs& Args);
 
+	/** Spawns the content bundle tab */
+	TSharedRef<SDockTab> SpawnContentBundleTab(const FSpawnTabArgs& Args);
+
 	bool BuildHLODs(const FString& InMapToProcess);
 	bool BuildMinimap(const FString& InMapToProcess);
 	bool BuildLandscapeSplineMeshes(UWorld* InWorld);
@@ -92,6 +100,10 @@ private:
 
 	TWeakPtr<SDockTab> WorldPartitionTab;
 
+	TWeakPtr<SDockTab> ContentBundleTab;
+
+	TWeakPtr<SContentBundleBrowser> ContentBundleBrowser;
+
 	FWorldPartitionCreated WorldPartitionCreatedEvent;
 
 	FOnExecuteCommandlet OnExecuteCommandletEvent;
diff --git a/Engine/Source/Editor/WorldPartitionEditor/WorldPartitionEditor.Build.cs b/Engine/Source/Editor/WorldPartitionEditor/WorldPartitionEditor.Build.cs
--- a/Engine/Source/Editor/WorldPartitionEditor/WorldPartitionEditor.Build.cs
+++ b/Engine/Source/Editor/WorldPartitionEditor/WorldPartitionEditor.Build.cs
@@ -25,6 +25,7 @@ public class WorldPartitionEditor : ModuleRules
 				"RenderCore",
 				"Renderer",
 				"RHI",
+				"SceneOutliner"
 			}
 		);
 
