Expose ContentBundleDescriptor and Editing functionalities to allow Drag and drog handler to toggle automatically content bundle editing when dropping a plugin actor

diff --git a/Engine/Plugins/Experimental/GameFeatures/Source/GameFeatures/Public/GameFeatureAction_AddWPContent.h b/Engine/Plugins/Experimental/GameFeatures/Source/GameFeatures/Public/GameFeatureAction_AddWPContent.h
--- a/Engine/Plugins/Experimental/GameFeatures/Source/GameFeatures/Public/GameFeatureAction_AddWPContent.h
+++ b/Engine/Plugins/Experimental/GameFeatures/Source/GameFeatures/Public/GameFeatureAction_AddWPContent.h
@@ -16,7 +16,7 @@ class UContentBundleDescriptor;
  *
  */
 UCLASS(meta = (DisplayName = "Add World Partition Content"))
-class UGameFeatureAction_AddWPContent : public UGameFeatureAction
+class GAMEFEATURES_API UGameFeatureAction_AddWPContent : public UGameFeatureAction
 {
 	GENERATED_UCLASS_BODY()
 
@@ -30,6 +30,8 @@ public:
 
 	UGameFeatureData* GetGameFeatureData() const;
 
+	const UContentBundleDescriptor* GetContentBundleDescriptor() const { return ContentBundleDescriptor; }
+
 private:
 	UPROPERTY(EditAnywhere, Category = ContentBundle)
 	TObjectPtr<UContentBundleDescriptor> ContentBundleDescriptor;
diff --git a/Engine/Source/Editor/WorldPartitionEditor/Private/WorldPartition/ContentBundle/ContentBundleEditorSubsystem.cpp b/Engine/Source/Editor/WorldPartitionEditor/Private/WorldPartition/ContentBundle/ContentBundleEditorSubsystem.cpp
--- a/Engine/Source/Editor/WorldPartitionEditor/Private/WorldPartition/ContentBundle/ContentBundleEditorSubsystem.cpp
+++ b/Engine/Source/Editor/WorldPartitionEditor/Private/WorldPartition/ContentBundle/ContentBundleEditorSubsystem.cpp
@@ -112,7 +112,7 @@ TSharedRef<SWidget> UContentBundleEditionSubmodule::GetActorEditorContextWidget(
 	return OutWidget;
 }
 
-void UContentBundleEditionSubmodule::ActivateContentBundleEditing(TSharedPtr<FContentBundleEditor>& ContentBundleEditor)
+bool UContentBundleEditionSubmodule::ActivateContentBundleEditing(TSharedPtr<FContentBundleEditor>& ContentBundleEditor)
 {
 	EContentBundleStatus ContentBundleStatus = ContentBundleEditor->GetStatus();
 	bool bCanActivateEditing = ContentBundleStatus == EContentBundleStatus::ContentInjected || ContentBundleStatus == EContentBundleStatus::ReadyToInject;
@@ -138,10 +138,24 @@ void UContentBundleEditionSubmodule::ActivateContentBundleEditing(TSharedPtr<FCo
 		UE_LOG(LogContentBundle, Log, TEXT("[CB: %s] Content Bundle is being edited"), *ContentBundleEditor->GetDescriptor()->GetDisplayName());
 
 		GetSubsystem()->NotifyContentBundleChanged(ContentBundleEditor.Get());
+
+		return true;
+	}
+	return false;
+}
+
+bool UContentBundleEditionSubmodule::DeactivateContentBundleEditing(TSharedPtr<FContentBundleEditor>& ContentBundleEditor)
+{
+	if (IsEditingContentBundle(ContentBundleEditor.Get()))
+	{
+		DeactivateCurrentContentBundleEditing();
+		return true;
 	}
+
+	return false;
 }
 
-void UContentBundleEditionSubmodule::DeactivateCurrentContentBundleEditing()
+bool UContentBundleEditionSubmodule::DeactivateCurrentContentBundleEditing()
 {
 	if (TSharedPtr<FContentBundleEditor> ContentBundleEditorPin = EditingContentBundle.Pin())
 	{
@@ -153,7 +167,11 @@ void UContentBundleEditionSubmodule::DeactivateCurrentContentBundleEditing()
 		GetSubsystem()->NotifyContentBundleChanged(ContentBundleEditorPin.Get());
 
 		UE_LOG(LogContentBundle, Log, TEXT("[CB: %s] Content Bundle is no longer being edited"), *ContentBundleEditorPin->GetDescriptor()->GetDisplayName());
+
+		return true;
 	}
+
+	return false;
 }
 
 void UContentBundleEditionSubmodule::PushContentBundleEditing()
@@ -225,6 +243,21 @@ TArray<TSharedPtr<FContentBundleEditor>> UContentBundleEditorSubsystem::GetEdito
 	return EditorContentBundles;
 }
 
+TSharedPtr<FContentBundleEditor> UContentBundleEditorSubsystem::GetEditorContentBundle(const UContentBundleDescriptor* ContentBundleDescriptor) const
+{
+	if (UContentBundleManager* ContentBundleManager = GetWorld()->ContentBundleManager)
+	{
+		return ContentBundleManager->GetEditorContentBundle(ContentBundleDescriptor, GetWorld());
+	}
+
+	return nullptr;
+}
+
+bool UContentBundleEditorSubsystem::HasContentBundle(const UContentBundleDescriptor* ContentBundleDescriptor) const
+{
+	return GetEditorContentBundle(ContentBundleDescriptor) != nullptr;
+}
+
 void UContentBundleEditorSubsystem::SelectActors(FContentBundleEditor& EditorContentBundle)
 {
 	SelectActorsInternal(EditorContentBundle, true);
@@ -245,6 +278,21 @@ void UContentBundleEditorSubsystem::UnreferenceAllActors(FContentBundleEditor& E
 	EditorContentBundle.UnreferenceAllActors();
 }
 
+bool UContentBundleEditorSubsystem::ActivateContentBundleEditing(TSharedPtr<FContentBundleEditor>& ContentBundleEditor)  const
+{
+	return ContentBundleEditionSubModule->ActivateContentBundleEditing(ContentBundleEditor);
+}
+
+bool UContentBundleEditorSubsystem::DectivateContentBundleEditing(TSharedPtr<FContentBundleEditor>& ContentBundleEditor) const
+{
+	return ContentBundleEditionSubModule->DeactivateContentBundleEditing(ContentBundleEditor);
+}
+
+bool UContentBundleEditorSubsystem::IsContentBundleEditingActivated(TSharedPtr<FContentBundleEditor>& ContentBundleEditor) const
+{
+	return ContentBundleEditionSubModule->IsEditingContentBundle(ContentBundleEditor.Get());
+}
+
 void UContentBundleEditorSubsystem::SelectActorsInternal(FContentBundleEditor& EditorContentBundle, bool bSelect)
 {
 	GEditor->GetSelectedActors()->BeginBatchSelectOperation();
diff --git a/Engine/Source/Editor/WorldPartitionEditor/Public/WorldPartition/ContentBundle/ContentBundleEditorSubsystem.h b/Engine/Source/Editor/WorldPartitionEditor/Public/WorldPartition/ContentBundle/ContentBundleEditorSubsystem.h
--- a/Engine/Source/Editor/WorldPartitionEditor/Public/WorldPartition/ContentBundle/ContentBundleEditorSubsystem.h
+++ b/Engine/Source/Editor/WorldPartitionEditor/Public/WorldPartition/ContentBundle/ContentBundleEditorSubsystem.h
@@ -11,7 +11,7 @@
 #include "ContentBundleEditorSubsystem.generated.h"
 
 class FContentBundleEditor;
-class FContentBundleEditor;
+class UContentBundleDescriptor;
 class UContentBundleEditorSubsystem;
 
 UCLASS(Within = ContentBundleEditorSubsystem)
@@ -53,8 +53,9 @@ public:
 	virtual FOnActorEditorContextClientChanged& GetOnActorEditorContextClientChanged() override { return ActorEditorContextClientChanged; }
 	//~ End IActorEditorContextClient interface
 
-	void ActivateContentBundleEditing(TSharedPtr<FContentBundleEditor>& ContentBundleEditor);
-	void DeactivateCurrentContentBundleEditing();
+	bool ActivateContentBundleEditing(TSharedPtr<FContentBundleEditor>& ContentBundleEditor);
+	bool DeactivateContentBundleEditing(TSharedPtr<FContentBundleEditor>& ContentBundleEditor);
+	bool DeactivateCurrentContentBundleEditing();
 	bool IsEditingContentBundle() const { return EditingContentBundle.IsValid(); }
 	bool IsEditingContentBundle(const FContentBundleEditor* ContentBundleEditor) const { return EditingContentBundle.Pin().Get() == ContentBundleEditor; }
 
@@ -76,7 +77,7 @@ private:
 };
 
 UCLASS()
-class UContentBundleEditorSubsystem : public UEditorSubsystem, public IContentBundleEditorSubsystemInterface
+class WORLDPARTITIONEDITOR_API UContentBundleEditorSubsystem : public UEditorSubsystem, public IContentBundleEditorSubsystemInterface
 {
 	GENERATED_BODY()
 
@@ -103,6 +104,9 @@ public:
 	TSharedPtr<FContentBundleEditor> GetEditorContentBundleForActor(const AActor* Actor);
 
 	TArray<TSharedPtr<FContentBundleEditor>> GetEditorContentBundles();
+	TSharedPtr<FContentBundleEditor> GetEditorContentBundle(const UContentBundleDescriptor* ContentBundleDescriptor) const;
+
+	bool HasContentBundle(const UContentBundleDescriptor* ContentBundleDescriptor) const;
 
 	void SelectActors(FContentBundleEditor& EditorContentBundle);
 	void DeselectActors(FContentBundleEditor& EditorContentBundle);
@@ -110,6 +114,10 @@ public:
 	void ReferenceAllActors(FContentBundleEditor& EditorContentBundle);
 	void UnreferenceAllActors(FContentBundleEditor& EditorContentBundle);
 
+	bool ActivateContentBundleEditing(TSharedPtr<FContentBundleEditor>& ContentBundleEditor) const;
+	bool DectivateContentBundleEditing(TSharedPtr<FContentBundleEditor>& ContentBundleEditor) const;
+	bool IsContentBundleEditingActivated(TSharedPtr<FContentBundleEditor>& ContentBundleEditor) const;
+
 	DECLARE_EVENT_OneParam(UContentBundleEditorSubsystem, FOnContentBundleChanged, const FContentBundleEditor*);
 	FOnContentBundleChanged& OnContentBundleChanged() { return ContentBundleChanged; }
 
