World Partition Editor: added load from here (mapped on control+double click) which will load around clicked area.

#rnx

diff --git a/Engine/Source/Editor/WorldPartitionEditor/Private/WorldPartition/SWorldPartitionEditorGrid2D.cpp b/Engine/Source/Editor/WorldPartitionEditor/Private/WorldPartition/SWorldPartitionEditorGrid2D.cpp
--- a/Engine/Source/Editor/WorldPartitionEditor/Private/WorldPartition/SWorldPartitionEditorGrid2D.cpp
+++ b/Engine/Source/Editor/WorldPartitionEditor/Private/WorldPartition/SWorldPartitionEditorGrid2D.cpp
@@ -235,6 +235,7 @@ void SWorldPartitionEditorGrid2D::FEditorCommands::RegisterCommands()
 	UI_COMMAND(ConvertSelectedRegionsToActors, "Convert Selected Regions To Actors", "Convert the selected regions to actors.", EUserInterfaceActionType::Button, FInputChord());
 	UI_COMMAND(MoveCameraHere, "Move Camera Here", "Move the camera to the selected position.", EUserInterfaceActionType::Button, FInputChord());
 	UI_COMMAND(PlayFromHere, "Play From Here", "Play from here.", EUserInterfaceActionType::Button, FInputChord());
+	UI_COMMAND(LoadFromHere, "Load From Here", "Load from here.", EUserInterfaceActionType::Button, FInputChord());
 }
 
 SWorldPartitionEditorGrid2D::SWorldPartitionEditorGrid2D()
@@ -442,6 +443,7 @@ void SWorldPartitionEditorGrid2D::Construct(const FArguments& InArgs)
 	ActionList.MapAction(Commands.ConvertSelectedRegionsToActors, FExecuteAction::CreateSP(this, &SWorldPartitionEditorGrid2D::ConvertSelectedRegionsToActors), FCanExecuteAction::CreateLambda(CanConvertSelectedRegionsToActors));
 	ActionList.MapAction(Commands.MoveCameraHere, FExecuteAction::CreateSP(this, &SWorldPartitionEditorGrid2D::MoveCameraHere));
 	ActionList.MapAction(Commands.PlayFromHere, FExecuteAction::CreateSP(this, &SWorldPartitionEditorGrid2D::PlayFromHere));
+	ActionList.MapAction(Commands.LoadFromHere, FExecuteAction::CreateSP(this, &SWorldPartitionEditorGrid2D::LoadFromHere));
 }
 
 void SWorldPartitionEditorGrid2D::UpdateWorldMiniMapDetails()
@@ -598,6 +600,33 @@ void SWorldPartitionEditorGrid2D::PlayFromHere()
 	}	
 }
 
+void SWorldPartitionEditorGrid2D::LoadFromHere()
+{
+	const float SelectionSnap = GetSelectionSnap();
+	const FVector LoadLocation = FVector(MouseCursorPosWorld, 0);
+	const FVector LoadExtent(SelectionSnap, SelectionSnap, HALF_WORLD_MAX);
+	FBox LoadCellsBox(LoadLocation - LoadExtent, LoadLocation + LoadExtent);
+
+	// Snap box
+	LoadCellsBox.Min.X = FMath::GridSnap(LoadCellsBox.Min.X - SelectionSnap / 2, SelectionSnap);
+	LoadCellsBox.Min.Y = FMath::GridSnap(LoadCellsBox.Min.Y - SelectionSnap / 2, SelectionSnap);
+	LoadCellsBox.Max.X = FMath::GridSnap(LoadCellsBox.Max.X + SelectionSnap / 2, SelectionSnap);
+	LoadCellsBox.Max.Y = FMath::GridSnap(LoadCellsBox.Max.Y + SelectionSnap / 2, SelectionSnap);
+
+	// Clip to minimap
+	if (WorldMiniMapBounds.bIsValid && WorldMiniMapBounds.IsInside(FVector2D(MouseCursorPosWorld)))
+	{
+		LoadCellsBox.Min.X = FMath::Max(LoadCellsBox.Min.X, WorldMiniMapBounds.Min.X);
+		LoadCellsBox.Min.Y = FMath::Max(LoadCellsBox.Min.Y, WorldMiniMapBounds.Min.Y);
+		LoadCellsBox.Max.X = FMath::Min(LoadCellsBox.Max.X, WorldMiniMapBounds.Max.X);
+		LoadCellsBox.Max.Y = FMath::Min(LoadCellsBox.Max.Y, WorldMiniMapBounds.Max.Y);
+	}
+
+	// Load box
+	UWorldPartitionEditorLoaderAdapter* EditorLoaderAdapter = WorldPartition->CreateEditorLoaderAdapter<FLoaderAdapterShape>(World, LoadCellsBox, TEXT("Loaded Region"));
+	EditorLoaderAdapter->GetLoaderAdapter()->Load();
+}
+
 bool SWorldPartitionEditorGrid2D::IsFollowPlayerInPIE() const
 {
 	return bFollowPlayerInPIE && (UWorldPartition::IsSimulating() || GEditor->PlayWorld);
@@ -691,6 +720,7 @@ FReply SWorldPartitionEditorGrid2D::OnMouseButtonUp(const FGeometry& MyGeometry,
 			MenuBuilder.BeginSection(NAME_None, LOCTEXT("WorldPartitionMisc", "Misc"));
 				MenuBuilder.AddMenuEntry(Commands.MoveCameraHere);
 				MenuBuilder.AddMenuEntry(Commands.PlayFromHere);
+				MenuBuilder.AddMenuEntry(Commands.LoadFromHere);
 			MenuBuilder.EndSection();			
 
 			FWidgetPath WidgetPath = MouseEvent.GetEventPath() != nullptr ? *MouseEvent.GetEventPath() : FWidgetPath();
@@ -761,6 +791,11 @@ FReply SWorldPartitionEditorGrid2D::OnMouseButtonDoubleClick(const FGeometry& In
 		else
 		{
 			MoveCameraHere();
+
+			if (InMouseEvent.IsControlDown())
+			{
+				LoadFromHere();
+			}
 		}
 	}
 	return FReply::Handled();
@@ -1150,7 +1185,7 @@ uint32 SWorldPartitionEditorGrid2D::PaintActors(const FGeometry& AllottedGeometr
 
 							if (!bIsLocalHovered)
 							{
-								if (UTexture2D* Texture2D = Cast<UTexture2D>(WorldMiniMapBrush.GetResourceObject()))
+								if (WorldMiniMapBounds.bIsValid)
 								{
 									const FBox2D MinimapBounds(
 										WorldToScreen.TransformPoint(WorldMiniMapBounds.Min),
diff --git a/Engine/Source/Editor/WorldPartitionEditor/Public/WorldPartition/SWorldPartitionEditorGrid2D.h b/Engine/Source/Editor/WorldPartitionEditor/Public/WorldPartition/SWorldPartitionEditorGrid2D.h
--- a/Engine/Source/Editor/WorldPartitionEditor/Public/WorldPartition/SWorldPartitionEditorGrid2D.h
+++ b/Engine/Source/Editor/WorldPartitionEditor/Public/WorldPartition/SWorldPartitionEditorGrid2D.h
@@ -25,6 +25,7 @@ protected:
 		TSharedPtr<FUICommandInfo> ConvertSelectedRegionsToActors;
 		TSharedPtr<FUICommandInfo> MoveCameraHere;
 		TSharedPtr<FUICommandInfo> PlayFromHere;
+		TSharedPtr<FUICommandInfo> LoadFromHere;
 
 		/**
 		 * Initialize commands
@@ -45,6 +46,7 @@ protected:
 	void ConvertSelectedRegionsToActors();
 	void MoveCameraHere();
 	void PlayFromHere();
+	void LoadFromHere();
 
 	bool IsFollowPlayerInPIE() const;
 	bool IsInteractive() const;
