Remove empty localized string

diff --git a/Engine/Source/Editor/WorldPartitionEditor/Private/WorldPartition/ContentBundle/Outliner/ContentBundleEditingColumn.cpp b/Engine/Source/Editor/WorldPartitionEditor/Private/WorldPartition/ContentBundle/Outliner/ContentBundleEditingColumn.cpp
--- a/Engine/Source/Editor/WorldPartitionEditor/Private/WorldPartition/ContentBundle/Outliner/ContentBundleEditingColumn.cpp
+++ b/Engine/Source/Editor/WorldPartitionEditor/Private/WorldPartition/ContentBundle/Outliner/ContentBundleEditingColumn.cpp
@@ -62,7 +62,7 @@ const TSharedRef<SWidget> FContentBundleOutlinerEditingColumn::ConstructRowWidge
 								return LOCTEXT("ContentBundleOutlinerEditingActive", "Active");
 							}
 							
-							return LOCTEXT("", ""); 
+							return FText();
 						})
 					.ColorAndOpacity(MakeAttributeLambda([ContentBundleTreeItem] { return ContentBundleTreeItem->GetItemColor(); }))
 				];
