World Partition Editor: added a measuring tool, mapped on middle-click drag.

#rnx

diff --git a/Engine/Source/Editor/WorldPartitionEditor/Private/WorldPartition/SWorldPartitionEditorGrid2D.cpp b/Engine/Source/Editor/WorldPartitionEditor/Private/WorldPartition/SWorldPartitionEditorGrid2D.cpp
--- a/Engine/Source/Editor/WorldPartitionEditor/Private/WorldPartition/SWorldPartitionEditorGrid2D.cpp
+++ b/Engine/Source/Editor/WorldPartitionEditor/Private/WorldPartition/SWorldPartitionEditorGrid2D.cpp
@@ -246,6 +246,7 @@ SWorldPartitionEditorGrid2D::SWorldPartitionEditorGrid2D()
 	, ScreenRect(ForceInit)
 	, bIsDragSelecting(false)
 	, bIsPanning(false)
+	, bIsMeasuring(false)
 	, bShowActors(false)
 	, bFollowPlayerInPIE(false)
 	, SelectBox(ForceInit)
@@ -653,7 +654,7 @@ FReply SWorldPartitionEditorGrid2D::OnMouseButtonDown(const FGeometry& MyGeometr
 
 	TotalMouseDelta = 0;
 
-	if (bIsLeftMouseButtonEffecting || bIsRightMouseButtonEffecting)
+	if (bIsLeftMouseButtonEffecting || bIsMiddleMouseButtonEffecting || bIsRightMouseButtonEffecting)
 	{
 		FReply ReplyState = FReply::Handled();
 		ReplyState.CaptureMouse(SharedThis(this));
@@ -671,6 +672,11 @@ FReply SWorldPartitionEditorGrid2D::OnMouseButtonDown(const FGeometry& MyGeometr
 			SelectBoxGridSnapped.Init();
 		}
 
+		if (bIsMiddleMouseButtonEffecting)
+		{
+			MeasureStart = MeasureEnd = MouseCursorPosWorld;
+		}
+
 		return ReplyState;
 	}
 
@@ -688,11 +694,11 @@ FReply SWorldPartitionEditorGrid2D::OnMouseButtonUp(const FGeometry& MyGeometry,
 
 	TotalMouseDelta = 0;
 
-	if (bIsLeftMouseButtonEffecting || bIsRightMouseButtonEffecting)
+	if (bIsLeftMouseButtonEffecting || bIsMiddleMouseButtonEffecting || bIsRightMouseButtonEffecting)
 	{
 		FReply ReplyState = FReply::Handled();
 
-		const bool bHasMouseCapture = bIsDragSelecting || bIsPanning;
+		const bool bHasMouseCapture = bIsDragSelecting || bIsPanning || bIsMeasuring;
 		MouseCursorPos = MyGeometry.AbsoluteToLocal(MouseEvent.GetScreenSpacePosition());
 		MouseCursorPosWorld = ScreenToWorld.TransformPoint(MouseCursorPos);
 
@@ -769,7 +775,12 @@ FReply SWorldPartitionEditorGrid2D::OnMouseButtonUp(const FGeometry& MyGeometry,
 			bIsPanning = false;
 		}
 
-		if (HasMouseCapture() && !bIsDragSelecting && !bIsPanning)
+		if (bIsMiddleMouseButtonEffecting)
+		{
+			bIsMeasuring = false;
+		}
+
+		if (HasMouseCapture() && !bIsDragSelecting && !bIsPanning && !bIsMeasuring)
 		{
 			ReplyState.ReleaseMouseCapture();
 		}
@@ -817,6 +828,19 @@ FReply SWorldPartitionEditorGrid2D::OnMouseMove(const FGeometry& MyGeometry, con
 		const bool bIsMiddleMouseButtonDown = MouseEvent.IsMouseButtonDown(EKeys::MiddleMouseButton);
 		const bool bIsDragTrigger = IsInteractive() && (TotalMouseDelta > FSlateApplication::Get().GetDragTriggerDistance());
 
+		if (bIsMiddleMouseButtonDown)
+		{
+			if (!bIsMeasuring && bIsDragTrigger)
+			{
+				bIsMeasuring = true;
+			}
+
+			if (bIsMeasuring)
+			{
+				MeasureEnd = MouseCursorPosWorld;
+			}
+		}
+
 		if (bIsLeftMouseButtonDown)
 		{
 			if (!bIsDragSelecting && bIsDragTrigger)
@@ -878,7 +902,7 @@ int32 SWorldPartitionEditorGrid2D::PaintGrid(const FGeometry& AllottedGeometry,
 	TArray<FVector2D> LinePoints;
 	LinePoints.SetNum(2);
 
-	FVector2D ScreenWorldOrigin = WorldToScreen.TransformPoint(FVector2D(0, 0));
+	const FVector2D ScreenWorldOrigin = WorldToScreen.TransformPoint(FVector2D(0, 0));
 	
 	// World Y-axis
 	if (ScreenWorldOrigin.X > ScreenRect.Min.X && ScreenWorldOrigin.X < ScreenRect.Max.X)
@@ -889,7 +913,16 @@ int32 SWorldPartitionEditorGrid2D::PaintGrid(const FGeometry& AllottedGeometry,
 		FLinearColor YAxisColor = FLinearColor::Green;
 		YAxisColor.A = 0.4f;
 		
-		FSlateDrawElement::MakeLines(OutDrawElements, LayerId, AllottedGeometry.ToPaintGeometry(), LinePoints, ESlateDrawEffect::None, YAxisColor, true, 2.0f);
+		FSlateDrawElement::MakeLines(
+			OutDrawElements, 
+			LayerId, 
+			AllottedGeometry.ToPaintGeometry(), 
+			LinePoints, 
+			ESlateDrawEffect::None, 
+			YAxisColor, 
+			true, 
+			2.0f
+		);
 	}
 
 	// World X-axis
@@ -901,7 +934,16 @@ int32 SWorldPartitionEditorGrid2D::PaintGrid(const FGeometry& AllottedGeometry,
 		FLinearColor XAxisColor = FLinearColor::Red;
 		XAxisColor.A = 0.4f;
 		
-		FSlateDrawElement::MakeLines(OutDrawElements, LayerId, AllottedGeometry.ToPaintGeometry(), LinePoints, ESlateDrawEffect::None, XAxisColor, true, 2.0f);
+		FSlateDrawElement::MakeLines(
+			OutDrawElements, 
+			LayerId, 
+			AllottedGeometry.ToPaintGeometry(), 
+			LinePoints, 
+			ESlateDrawEffect::None, 
+			XAxisColor, 
+			true, 
+			2.0f
+		);
 	}
 
 	return LayerId + 1;
@@ -1098,28 +1140,7 @@ uint32 SWorldPartitionEditorGrid2D::PaintActors(const FGeometry& AllottedGeometr
 
 			if (LabelColorGradient > 0.0f)
 			{
-				const FLinearColor LabelForegroundColor(1.0f, 1.0f, 1.0f, Color.A * LabelColorGradient);
-				const FLinearColor LabelShadowColor(0, 0, 0, Color.A * LabelColorGradient);
-
-				FSlateDrawElement::MakeText(
-					OutDrawElements,
-					++LayerId,
-					AllottedGeometry.ToPaintGeometry(LabelTextPos + FVector2D(2, 2), FVector2D(1, 1)),
-					Label,
-					Font,
-					ESlateDrawEffect::None,
-					LabelShadowColor
-				);
-
-				FSlateDrawElement::MakeText(
-					OutDrawElements,
-					++LayerId,
-					AllottedGeometry.ToPaintGeometry(LabelTextPos, FVector2D(1,1)),
-					Label,
-					Font,
-					ESlateDrawEffect::None,
-					LabelForegroundColor
-				);
+				LayerId = DrawTextLabel(OutDrawElements, LayerId, AllottedGeometry, Label, Pos, Color.CopyWithNewOpacity(Color.A * LabelColorGradient), Font);
 			}
 		}
 	};
@@ -1606,6 +1627,7 @@ int32 SWorldPartitionEditorGrid2D::OnPaint(const FPaintArgs& Args, const FGeomet
 		LayerId = PaintViewer(AllottedGeometry, MyCullingRect, OutDrawElements, ++LayerId);
 		LayerId = PaintSelection(AllottedGeometry, MyCullingRect, OutDrawElements, ++LayerId);
 		LayerId = PaintSoftwareCursor(AllottedGeometry, MyCullingRect, OutDrawElements, ++LayerId);
+		LayerId = PaintMeasureTool(AllottedGeometry, MyCullingRect, OutDrawElements, ++LayerId);
 		
 		// Draw a surrounding indicator when PIE is active
 		if (UWorldPartition::IsSimulating() || !!GEditor->PlayWorld)
@@ -1766,6 +1788,78 @@ int32 SWorldPartitionEditorGrid2D::PaintMinimap(const FGeometry& AllottedGeometr
 	return LayerId;
 }
 
+int32 SWorldPartitionEditorGrid2D::PaintMeasureTool(const FGeometry& AllottedGeometry, const FSlateRect& MyCullingRect, FSlateWindowElementList& OutDrawElements, int32 LayerId) const
+{
+	// Measure tool
+	if (bIsMeasuring)
+	{
+		TArray<FVector2D> LinePoints;
+		LinePoints.SetNum(2);
+
+		const FVector2D MeasureStartScreen = WorldToScreen.TransformPoint(MeasureStart);
+		const FVector2D MeasureEndScreen = WorldToScreen.TransformPoint(MeasureEnd);
+		LinePoints[0] = MeasureStartScreen;
+		LinePoints[1] = MeasureEndScreen;
+
+		FSlateDrawElement::MakeLines(
+			OutDrawElements, 
+			LayerId, 
+			AllottedGeometry.ToPaintGeometry(), 
+			LinePoints, 
+			ESlateDrawEffect::None, 
+			FLinearColor::Yellow, 
+			true, 
+			2.0f
+		);
+
+		DrawTextLabel(
+			OutDrawElements, 
+			LayerId, 
+			AllottedGeometry,
+			FString::Printf(TEXT("%d"), (int32)FVector2D::Distance(MeasureStart, MeasureEnd)),
+			(MeasureStartScreen + MeasureEndScreen) * 0.5f,
+			FLinearColor::White,
+			SmallLayoutFont
+		);
+	}
+
+	return LayerId;
+}
+
+int32 SWorldPartitionEditorGrid2D::DrawTextLabel(FSlateWindowElementList& OutDrawElements, int32 LayerId, const FGeometry& AllottedGeometry, const FString& Label, const FVector2D& Pos, const FLinearColor& Color, const FSlateFontInfo& Font) const
+{
+	const FVector2D LabelTextSize = FSlateApplication::Get().GetRenderer()->GetFontMeasureService()->Measure(Label, Font);
+
+	if (LabelTextSize.X > 0)
+	{
+		const FVector2D LabelTextPos = Pos - LabelTextSize * 0.5f;
+		const FLinearColor LabelForegroundColor(1.0f, 1.0f, 1.0f, Color.A);
+		const FLinearColor LabelShadowColor(0, 0, 0, Color.A);
+
+		FSlateDrawElement::MakeText(
+			OutDrawElements,
+			++LayerId,
+			AllottedGeometry.ToPaintGeometry(LabelTextPos + FVector2D(2, 2), FVector2D(1, 1)),
+			Label,
+			Font,
+			ESlateDrawEffect::None,
+			LabelShadowColor
+		);
+
+		FSlateDrawElement::MakeText(
+			OutDrawElements,
+			++LayerId,
+			AllottedGeometry.ToPaintGeometry(LabelTextPos, FVector2D(1,1)),
+			Label,
+			Font,
+			ESlateDrawEffect::None,
+			LabelForegroundColor
+		);
+	}
+
+	return LayerId;
+}
+
 FReply SWorldPartitionEditorGrid2D::FocusSelection()
 {
 	FBox SelectionBox(ForceInit);
diff --git a/Engine/Source/Editor/WorldPartitionEditor/Public/WorldPartition/SWorldPartitionEditorGrid2D.h b/Engine/Source/Editor/WorldPartitionEditor/Public/WorldPartition/SWorldPartitionEditorGrid2D.h
--- a/Engine/Source/Editor/WorldPartitionEditor/Public/WorldPartition/SWorldPartitionEditorGrid2D.h
+++ b/Engine/Source/Editor/WorldPartitionEditor/Public/WorldPartition/SWorldPartitionEditorGrid2D.h
@@ -69,6 +69,9 @@ protected:
 	virtual uint32 PaintSelection(const FGeometry& AllottedGeometry, const FSlateRect& MyCullingRect, FSlateWindowElementList& OutDrawElements, uint32 LayerId) const;
 	virtual int32 PaintSoftwareCursor(const FGeometry& AllottedGeometry, const FSlateRect& MyCullingRect, FSlateWindowElementList& OutDrawElements, int32 LayerId) const;
 	virtual int32 PaintMinimap(const FGeometry& AllottedGeometry, const FSlateRect& MyCullingRect, FSlateWindowElementList& OutDrawElements, int32 LayerId) const;
+	virtual int32 PaintMeasureTool(const FGeometry& AllottedGeometry, const FSlateRect& MyCullingRect, FSlateWindowElementList& OutDrawElements, int32 LayerId) const;
+
+	int32 DrawTextLabel(FSlateWindowElementList& OutDrawElements, int32 LayerId, const FGeometry& AllottedGeometry, const FString& Label, const FVector2D& Pos, const FLinearColor& Color, const FSlateFontInfo& Font) const;
 
 	virtual FReply FocusSelection();
 	virtual FReply FocusLoadedRegions();
@@ -102,6 +105,7 @@ protected:
 
 	bool bIsDragSelecting;
 	bool bIsPanning;
+	bool bIsMeasuring;
 	bool bShowActors;
 	bool bFollowPlayerInPIE;
 	FVector2D MouseCursorPos;
@@ -109,6 +113,8 @@ protected:
 	FVector2D LastMouseCursorPosWorldDrag;
 	FVector2D SelectionStart;
 	FVector2D SelectionEnd;
+	FVector2D MeasureStart;
+	FVector2D MeasureEnd;
 	FBox SelectBox;
 	FBox SelectBoxGridSnapped;
 	FSlateFontInfo SmallLayoutFont;
