WorldPartitionEditorModule: Allow running external builders

diff --git a/Engine/Source/Editor/UnrealEd/Private/WorldPartition/WorldPartitionBuilder.cpp b/Engine/Source/Editor/UnrealEd/Private/WorldPartition/WorldPartitionBuilder.cpp
--- a/Engine/Source/Editor/UnrealEd/Private/WorldPartition/WorldPartitionBuilder.cpp
+++ b/Engine/Source/Editor/UnrealEd/Private/WorldPartition/WorldPartitionBuilder.cpp
@@ -20,7 +20,7 @@
 #include "UObject/SavePackage.h"
 #include "Algo/Transform.h"
 
-DEFINE_LOG_CATEGORY_STATIC(LogWorldPartitionBuilder, Log, All);
+DEFINE_LOG_CATEGORY_STATIC(LogWorldPartitionBuilder, All, All);
 
 namespace FWorldPartitionBuilderLog
 {
diff --git a/Engine/Source/Editor/WorldPartitionEditor/Private/WorldPartitionEditorModule.cpp b/Engine/Source/Editor/WorldPartitionEditor/Private/WorldPartitionEditorModule.cpp
--- a/Engine/Source/Editor/WorldPartitionEditor/Private/WorldPartitionEditorModule.cpp
+++ b/Engine/Source/Editor/WorldPartitionEditor/Private/WorldPartitionEditorModule.cpp
@@ -466,29 +466,37 @@ bool FWorldPartitionEditorModule::ConvertMap(const FString& InLongPackageName)
 	return false;
 }
 
-bool FWorldPartitionEditorModule::RunBuilder(TSubclassOf<UWorldPartitionBuilder> InWorldPartitionBuilder, UWorld* InWorld)
+bool IWorldPartitionEditorModule::RunBuilder(TSubclassOf<UWorldPartitionBuilder> InWorldPartitionBuilder, UWorld* InWorld)
+{
+	FRunBuilderParams Params;
+	Params.BuilderClass = InWorldPartitionBuilder;
+	Params.World = InWorld;
+	return RunBuilder(Params);
+}
+
+bool FWorldPartitionEditorModule::RunBuilder(const FRunBuilderParams& InParams)
 {
 	// Ideally this should be improved to automatically register all builders & present their options in a consistent way...
 
-	if (InWorldPartitionBuilder == UWorldPartitionHLODsBuilder::StaticClass())
+	if (InParams.BuilderClass == UWorldPartitionHLODsBuilder::StaticClass())
 	{
-		return BuildHLODs(InWorld->GetPackage()->GetName());
+		return BuildHLODs(InParams);
 	}
 	
-	if (InWorldPartitionBuilder == UWorldPartitionMiniMapBuilder::StaticClass())
+	if (InParams.BuilderClass == UWorldPartitionMiniMapBuilder::StaticClass())
 	{
-		return BuildMinimap(InWorld->GetPackage()->GetName());
+		return BuildMinimap(InParams);
 	}
 
-	if (InWorldPartitionBuilder == UWorldPartitionLandscapeSplineMeshesBuilder::StaticClass())
+	if (InParams.BuilderClass == UWorldPartitionLandscapeSplineMeshesBuilder::StaticClass())
 	{
-		return BuildLandscapeSplineMeshes(InWorld);
+		return BuildLandscapeSplineMeshes(InParams.World);
 	}
 
-	return false;
+	return Build(InParams);
 }
 
-bool FWorldPartitionEditorModule::BuildHLODs(const FString& InMapToProcess)
+bool FWorldPartitionEditorModule::BuildHLODs(const FRunBuilderParams& InParams)
 {
 	TSharedPtr<SWindow> DlgWindow =
 		SNew(SWindow)
@@ -510,57 +518,53 @@ bool FWorldPartitionEditorModule::BuildHLODs(const FString& InMapToProcess)
 
 	if (BuildHLODsDialog->GetDialogResult() != SWorldPartitionBuildHLODsDialog::DialogResult::Cancel)
 	{
-		FString MapPackage = InMapToProcess;
-		if (!UnloadCurrentMap(/*bAskSaveContentPackages=*/true, MapPackage))
-		{
-			return false;
-		}
-
-		const FString BuildArgs = BuildHLODsDialog->GetDialogResult() == SWorldPartitionBuildHLODsDialog::DialogResult::BuildHLODs ? "-SetupHLODs -BuildHLODs -AllowCommandletRendering" : "-DeleteHLODs";
-		const FString CommandletArgs = MapPackage + " -run=WorldPartitionBuilderCommandlet -Builder=WorldPartitionHLODsBuilder " + BuildArgs;
-		const FText OperationDescription = LOCTEXT("HLODBuildProgress", "Building HLODs...");
-
-		int32 Result;
-		bool bCancelled;
-		FString CommandletOutput;
-		RunCommandletAsExternalProcess(CommandletArgs, OperationDescription, Result, bCancelled, CommandletOutput);
-
-		bool bSuccess = !bCancelled && Result == 0;
-		if (bSuccess)
-		{
-			RescanAssetsAndLoadMap(MapPackage);
-		}
-		else if (bCancelled)
-		{
-			FMessageDialog::Open(EAppMsgType::Ok, LOCTEXT("HLODBuildCancelled", "HLOD build cancelled!"));
-		}
-		else if (Result != 0)
-		{
-			FMessageDialog::Open(EAppMsgType::Ok, LOCTEXT("HLODBuildFailed", "HLOD build failed! See log for details."));
-		}
+		FRunBuilderParams ParamsCopy(InParams);
+		ParamsCopy.ExtraArgs = BuildHLODsDialog->GetDialogResult() == SWorldPartitionBuildHLODsDialog::DialogResult::BuildHLODs ? "-SetupHLODs -BuildHLODs -AllowCommandletRendering" : "-DeleteHLODs";
+		ParamsCopy.OperationDescription = LOCTEXT("HLODBuildProgress", "Building HLODs...");
 
-		return bSuccess;
+		return Build(ParamsCopy);
 	}
 
 	return false;
 }
 
-bool FWorldPartitionEditorModule::BuildMinimap(const FString& InMapToProcess)
+
+bool FWorldPartitionEditorModule::BuildMinimap(const FRunBuilderParams& InParams)
+{
+	FRunBuilderParams ParamsCopy(InParams);
+	ParamsCopy.ExtraArgs = "-AllowCommandletRendering";
+	ParamsCopy.OperationDescription = LOCTEXT("MinimapBuildProgress", "Building minimap...");
+	return Build(ParamsCopy);
+}
+
+bool FWorldPartitionEditorModule::Build(const FRunBuilderParams& InParams)
 {
-	FString MapPackage = InMapToProcess;
+	FString MapPackage = InParams.World->GetPackage()->GetName();
 	if (!UnloadCurrentMap(/*bAskSaveContentPackages=*/true, MapPackage))
 	{
 		return false;
 	}
 
-	const FString CommandletArgs = MapPackage + " -run=WorldPartitionBuilderCommandlet -Builder=WorldPartitionMinimapBuilder -AllowCommandletRendering";
-	const FText OperationDescription = LOCTEXT("MinimapBuildProgress", "Building minimap...");
+	TStringBuilder<512> CommandletArgsBuilder;
+	CommandletArgsBuilder.Append(MapPackage);
+	CommandletArgsBuilder.Append(" -run=WorldPartitionBuilderCommandlet -Builder=");
+	CommandletArgsBuilder.Append(InParams.BuilderClass->GetName());
+		
+	if (!InParams.ExtraArgs.IsEmpty())
+	{
+		CommandletArgsBuilder.Append(" ");
+		CommandletArgsBuilder.Append(InParams.ExtraArgs);
+	}
+
+	const FText OperationDescription = InParams.OperationDescription.IsEmptyOrWhitespace() ? LOCTEXT("BuildProgress", "Building...") : InParams.OperationDescription;
 
 	int32 Result;
 	bool bCancelled;
 	FString CommandletOutput;
+
+	FString CommandletArgs(CommandletArgsBuilder.ToString());
 	RunCommandletAsExternalProcess(CommandletArgs, OperationDescription, Result, bCancelled, CommandletOutput);
-	
+
 	bool bSuccess = !bCancelled && Result == 0;
 	if (bSuccess)
 	{
@@ -568,11 +572,11 @@ bool FWorldPartitionEditorModule::BuildMinimap(const FString& InMapToProcess)
 	}
 	else if (bCancelled)
 	{
-		FMessageDialog::Open(EAppMsgType::Ok, LOCTEXT("MinimapBuildCancelled", "Minimap build cancelled!"));
+		FMessageDialog::Open(EAppMsgType::Ok, LOCTEXT("BuildCancelled", "Build cancelled!"));
 	}
 	else if (Result != 0)
 	{
-		FMessageDialog::Open(EAppMsgType::Ok, LOCTEXT("MinimapBuildFailed", "Minimap build failed! See log for details."));
+		FMessageDialog::Open(EAppMsgType::Ok, LOCTEXT("BuildFailed", "Build failed! See log for details."));
 	}
 
 	return bSuccess;
diff --git a/Engine/Source/Editor/WorldPartitionEditor/Public/WorldPartitionEditorModule.h b/Engine/Source/Editor/WorldPartitionEditor/Public/WorldPartitionEditorModule.h
--- a/Engine/Source/Editor/WorldPartitionEditor/Public/WorldPartitionEditorModule.h
+++ b/Engine/Source/Editor/WorldPartitionEditor/Public/WorldPartitionEditorModule.h
@@ -55,7 +55,7 @@ public:
 	/**
 	 * Run a world partition builder for the current map. Will display a dialog to specify options for the generation.
 	 */
-	virtual bool RunBuilder(TSubclassOf<UWorldPartitionBuilder> WorldPartitionBuilder, UWorld* InWorld) override;
+	virtual bool RunBuilder(const FRunBuilderParams& InParams) override;
 
 	/** Return the world added event. */
 	virtual FWorldPartitionCreated& OnWorldPartitionCreated() override { return WorldPartitionCreatedEvent; }
@@ -86,8 +86,9 @@ private:
 	/** Spawns the content bundle tab */
 	TSharedRef<SDockTab> SpawnContentBundleTab(const FSpawnTabArgs& Args);
 
-	bool BuildHLODs(const FString& InMapToProcess);
-	bool BuildMinimap(const FString& InMapToProcess);
+	bool Build(const FRunBuilderParams& InParams);
+	bool BuildMinimap(const FRunBuilderParams& InParams);
+	bool BuildHLODs(const FRunBuilderParams& InParams);
 	bool BuildLandscapeSplineMeshes(UWorld* InWorld);
 
 private:
diff --git a/Engine/Source/Runtime/Engine/Public/WorldPartition/IWorldPartitionEditorModule.h b/Engine/Source/Runtime/Engine/Public/WorldPartition/IWorldPartitionEditorModule.h
--- a/Engine/Source/Runtime/Engine/Public/WorldPartition/IWorldPartitionEditorModule.h
+++ b/Engine/Source/Runtime/Engine/Public/WorldPartition/IWorldPartitionEditorModule.h
@@ -19,8 +19,17 @@ public:
 
 	UE_DEPRECATED(5.1, "Use RunBuilder with UWorld* instead.")
 	virtual bool RunBuilder(TSubclassOf<UWorldPartitionBuilder> BuilderClass, const FString& InLongPackageName) { return false; }
-
-	virtual bool RunBuilder(TSubclassOf<UWorldPartitionBuilder> BuilderClass, UWorld* InWorld) = 0;
+		
+	struct FRunBuilderParams
+	{
+		TSubclassOf<UWorldPartitionBuilder> BuilderClass;
+		UWorld* World;
+		FString ExtraArgs;
+		FText OperationDescription;
+	};
+
+	virtual bool RunBuilder(TSubclassOf<UWorldPartitionBuilder> BuilderClass, UWorld* InWorld);
+	virtual bool RunBuilder(const FRunBuilderParams& Params) = 0;
 
 	virtual int32 GetPlacementGridSize() const = 0;
 	virtual int32 GetInstancedFoliageGridSize() const = 0;
