Do not push an ActorEditorContext when editing a ContentBundle anymore
Display the content bundle as uninteractablel when editing a world they are not a part of

diff --git a/Engine/Source/Editor/WorldPartitionEditor/Private/WorldPartition/ContentBundle/ContentBundleEditorSubsystem.cpp b/Engine/Source/Editor/WorldPartitionEditor/Private/WorldPartition/ContentBundle/ContentBundleEditorSubsystem.cpp
--- a/Engine/Source/Editor/WorldPartitionEditor/Private/WorldPartition/ContentBundle/ContentBundleEditorSubsystem.cpp
+++ b/Engine/Source/Editor/WorldPartitionEditor/Private/WorldPartition/ContentBundle/ContentBundleEditorSubsystem.cpp
@@ -124,7 +124,6 @@ bool UContentBundleEditionSubmodule::ActivateContentBundleEditing(TSharedPtr<FCo
 			DeactivateCurrentContentBundleEditing();
 		}
 
-		UActorEditorContextSubsystem::Get()->PushContext();
 		EditingContentBundle = ContentBundleEditor;
 
 		ContentBundleEditor->StartEditing();
@@ -158,7 +157,6 @@ bool UContentBundleEditionSubmodule::DeactivateCurrentContentBundleEditing()
 		ContentBundleEditorPin->StopEditing();
 		EditingContentBundle = nullptr;
 		ActorEditorContextClientChanged.Broadcast(this);
-		UActorEditorContextSubsystem::Get()->PopContext();
 
 		GetSubsystem()->NotifyContentBundleChanged(ContentBundleEditorPin.Get());
 
diff --git a/Engine/Source/Editor/WorldPartitionEditor/Private/WorldPartition/ContentBundle/Outliner/ContentBundleHierarchy.cpp b/Engine/Source/Editor/WorldPartitionEditor/Private/WorldPartition/ContentBundle/Outliner/ContentBundleHierarchy.cpp
--- a/Engine/Source/Editor/WorldPartitionEditor/Private/WorldPartition/ContentBundle/Outliner/ContentBundleHierarchy.cpp
+++ b/Engine/Source/Editor/WorldPartitionEditor/Private/WorldPartition/ContentBundle/Outliner/ContentBundleHierarchy.cpp
@@ -35,7 +35,7 @@ void FContentBundleHiearchy::CreateItems(TArray<FSceneOutlinerTreeItemPtr>& OutI
 	{
 		for (TSharedPtr<FContentBundleEditor>& ContentBundleEditor : ContentBundleEditorSubsystem->GetEditorContentBundles())
 		{
-			if (FSceneOutlinerTreeItemPtr ContentBundleTreeItem = Mode->CreateItemFor<FContentBundleTreeItem>(ContentBundleEditor, true))
+			if (FSceneOutlinerTreeItemPtr ContentBundleTreeItem = Mode->CreateItemFor<FContentBundleTreeItem>(FContentBundleTreeItem::FInitializationValues(ContentBundleEditor, *StaticCast<FContentBundleMode*>(Mode)), true))
 			{
 				OutItems.Add(ContentBundleTreeItem);
 			}
diff --git a/Engine/Source/Editor/WorldPartitionEditor/Private/WorldPartition/ContentBundle/Outliner/ContentBundleMode.cpp b/Engine/Source/Editor/WorldPartitionEditor/Private/WorldPartition/ContentBundle/Outliner/ContentBundleMode.cpp
--- a/Engine/Source/Editor/WorldPartitionEditor/Private/WorldPartition/ContentBundle/Outliner/ContentBundleMode.cpp
+++ b/Engine/Source/Editor/WorldPartitionEditor/Private/WorldPartition/ContentBundle/Outliner/ContentBundleMode.cpp
@@ -92,7 +92,14 @@ UWorld* FContentBundleMode::GetEditingWorld() const
 {
 	if (UContentBundleEditorSubsystem* ContentBundleEditorSubsystem = UContentBundleEditorSubsystem::Get())
 	{
-		return ContentBundleEditorSubsystem->GetWorld();
+		if (UWorld* SubsystemWorld = ContentBundleEditorSubsystem->GetWorld())
+		{
+			ULevel* CurrentLevel = SubsystemWorld->GetCurrentLevel();
+			if (CurrentLevel != nullptr)
+			{
+				return CurrentLevel->GetTypedOuter<UWorld>();
+			}
+		}
 	}
 
 	return nullptr;
diff --git a/Engine/Source/Editor/WorldPartitionEditor/Private/WorldPartition/ContentBundle/Outliner/ContentBundleTreeItem.cpp b/Engine/Source/Editor/WorldPartitionEditor/Private/WorldPartition/ContentBundle/Outliner/ContentBundleTreeItem.cpp
--- a/Engine/Source/Editor/WorldPartitionEditor/Private/WorldPartition/ContentBundle/Outliner/ContentBundleTreeItem.cpp
+++ b/Engine/Source/Editor/WorldPartitionEditor/Private/WorldPartition/ContentBundle/Outliner/ContentBundleTreeItem.cpp
@@ -5,15 +5,17 @@
 #include "WorldPartition/ContentBundle/ContentBundle.h"
 #include "WorldPartition/ContentBundle/ContentBundleEditor.h"
 #include "WorldPartition/ContentBundle/ContentBundleEditorSubsystem.h"
+#include "WorldPartition/ContentBundle/Outliner/ContentBundleMode.h"
 #include "Widgets/Text/STextBlock.h"
 
 #define LOCTEXT_NAMESPACE "ContentBundle"
 
 const FSceneOutlinerTreeItemType FContentBundleTreeItem::Type(&ISceneOutlinerTreeItem::Type);
 
-FContentBundleTreeItem::FContentBundleTreeItem(const TWeakPtr<FContentBundleEditor>& InContentBundleEditor)
+FContentBundleTreeItem::FContentBundleTreeItem(FInitializationValues InitializationValues)
 	: ISceneOutlinerTreeItem(Type)
-	, ContentBundleEditor(InContentBundleEditor)
+	, ContentBundleEditor(InitializationValues.ContentBundleEditor)
+	, Mode(InitializationValues.Mode)
 {
 	
 }
@@ -50,10 +52,7 @@ bool FContentBundleTreeItem::CanInteract() const
 	TSharedPtr<FContentBundleEditor> ContentBundleEditorPin = ContentBundleEditor.Pin();
 	if (ContentBundleEditorPin != nullptr)
 	{
-		if (UContentBundleEditorSubsystem* ContentBundleEditorSubsystem = UContentBundleEditorSubsystem::Get())
-		{
-			return ContentBundleEditorPin->GetInjectedWorld() == ContentBundleEditorSubsystem->GetWorld();
-		}
+		return ContentBundleEditorPin->GetInjectedWorld() == Mode.GetEditingWorld();
 	}
 	
 	return false;
@@ -61,7 +60,6 @@ bool FContentBundleTreeItem::CanInteract() const
 
 TSharedRef<SWidget> FContentBundleTreeItem::GenerateLabelWidget(ISceneOutliner& Outliner, const STableRow<FSceneOutlinerTreeItemPtr>& InRow)
 {
-	
 	return SNew(SHorizontalBox)
 		+ SHorizontalBox::Slot()
 		.VAlign(VAlign_Center)
@@ -69,9 +67,9 @@ TSharedRef<SWidget> FContentBundleTreeItem::GenerateLabelWidget(ISceneOutliner&
 			SNew(STextBlock)
 			.Text(MakeAttributeLambda([this] { return  FText::FromString(GetDisplayString()); }))
 			.ColorAndOpacity(MakeAttributeLambda([this] 
-				{ 
-					return GetItemColor();
-				}))
+			{ 
+				return GetItemColor();
+			}))
 		];
 }
 
diff --git a/Engine/Source/Editor/WorldPartitionEditor/Public/WorldPartition/ContentBundle/Outliner/ContentBundleTreeItem.h b/Engine/Source/Editor/WorldPartitionEditor/Public/WorldPartition/ContentBundle/Outliner/ContentBundleTreeItem.h
--- a/Engine/Source/Editor/WorldPartitionEditor/Public/WorldPartition/ContentBundle/Outliner/ContentBundleTreeItem.h
+++ b/Engine/Source/Editor/WorldPartitionEditor/Public/WorldPartition/ContentBundle/Outliner/ContentBundleTreeItem.h
@@ -8,11 +8,25 @@
 #include "ISceneOutlinerTreeItem.h"
 
 class FContentBundleEditor;
+class FContentBundleMode;
 
 struct FContentBundleTreeItem : ISceneOutlinerTreeItem
 {
 public:
-	FContentBundleTreeItem(const TWeakPtr<FContentBundleEditor>& InContentBundleEditor);
+	struct FInitializationValues
+	{
+		FInitializationValues(const TWeakPtr<FContentBundleEditor>& InContentBundleEditor, const FContentBundleMode& InMode)
+			:ContentBundleEditor(InContentBundleEditor),
+			Mode(InMode)
+		{
+
+		}
+
+		const TWeakPtr<FContentBundleEditor> ContentBundleEditor;
+		const FContentBundleMode& Mode;
+	};
+
+	FContentBundleTreeItem(FInitializationValues InitializationValues);
 	TWeakPtr<FContentBundleEditor> GetContentBundleEditor() const { return IsValid() ? ContentBundleEditor : nullptr; }
 	TSharedPtr<FContentBundleEditor> GetContentBundleEditorPin() const { return IsValid() ? ContentBundleEditor.Pin() : nullptr; }
 
@@ -43,4 +57,6 @@ public:
 
 private:
 	TWeakPtr<FContentBundleEditor> ContentBundleEditor;
+
+	const FContentBundleMode& Mode;
 };
\ No newline at end of file
