Fix static analysis

diff --git a/Engine/Source/Editor/WorldPartitionEditor/Private/WorldPartition/ContentBundle/Outliner/ContentBundleTreeItem.cpp b/Engine/Source/Editor/WorldPartitionEditor/Private/WorldPartition/ContentBundle/Outliner/ContentBundleTreeItem.cpp
--- a/Engine/Source/Editor/WorldPartitionEditor/Private/WorldPartition/ContentBundle/Outliner/ContentBundleTreeItem.cpp
+++ b/Engine/Source/Editor/WorldPartitionEditor/Private/WorldPartition/ContentBundle/Outliner/ContentBundleTreeItem.cpp
@@ -12,7 +12,7 @@
 
 const FSceneOutlinerTreeItemType FContentBundleTreeItem::Type(&ISceneOutlinerTreeItem::Type);
 
-FContentBundleTreeItem::FContentBundleTreeItem(FInitializationValues InitializationValues)
+FContentBundleTreeItem::FContentBundleTreeItem(const FInitializationValues& InitializationValues)
 	: ISceneOutlinerTreeItem(Type)
 	, ContentBundleEditor(InitializationValues.ContentBundleEditor)
 	, Mode(InitializationValues.Mode)
diff --git a/Engine/Source/Editor/WorldPartitionEditor/Public/WorldPartition/ContentBundle/Outliner/ContentBundleTreeItem.h b/Engine/Source/Editor/WorldPartitionEditor/Public/WorldPartition/ContentBundle/Outliner/ContentBundleTreeItem.h
--- a/Engine/Source/Editor/WorldPartitionEditor/Public/WorldPartition/ContentBundle/Outliner/ContentBundleTreeItem.h
+++ b/Engine/Source/Editor/WorldPartitionEditor/Public/WorldPartition/ContentBundle/Outliner/ContentBundleTreeItem.h
@@ -26,7 +26,7 @@ public:
 		const FContentBundleMode& Mode;
 	};
 
-	FContentBundleTreeItem(FInitializationValues InitializationValues);
+	FContentBundleTreeItem(const FInitializationValues& InitializationValues);
 	TWeakPtr<FContentBundleEditor> GetContentBundleEditor() const { return IsValid() ? ContentBundleEditor : nullptr; }
 	TSharedPtr<FContentBundleEditor> GetContentBundleEditorPin() const { return IsValid() ? ContentBundleEditor.Pin() : nullptr; }
 
