World Partition:
- Added option to always load actors in the editor even if streaming is enabled.
- Disabled world partition editor commands that act on loading regions when loading in the editor is disabled.
- Added option to not set the "can be used by level instance" flag when converting worlds with streaming disabled.

#rnx

diff --git a/Engine/Source/Editor/LevelInstanceEditor/Private/LevelInstanceEditorModule.cpp b/Engine/Source/Editor/LevelInstanceEditor/Private/LevelInstanceEditorModule.cpp
--- a/Engine/Source/Editor/LevelInstanceEditor/Private/LevelInstanceEditorModule.cpp
+++ b/Engine/Source/Editor/LevelInstanceEditor/Private/LevelInstanceEditorModule.cpp
@@ -596,6 +596,7 @@ namespace LevelInstanceMenuUtils
 		FWorldPartitionConverter::FParameters Parameters;
 		Parameters.bConvertSubLevels = false;
 		Parameters.bEnableStreaming = false;
+		Parameters.bCanBeUsedByLevelInstance = true;
 		Parameters.bUseActorFolders = true;
 
 		if (FWorldPartitionConverter::Convert(WorldAsset, Parameters))
diff --git a/Engine/Source/Editor/UnrealEd/Private/SWorldPartitionViewportWidget.cpp b/Engine/Source/Editor/UnrealEd/Private/SWorldPartitionViewportWidget.cpp
--- a/Engine/Source/Editor/UnrealEd/Private/SWorldPartitionViewportWidget.cpp
+++ b/Engine/Source/Editor/UnrealEd/Private/SWorldPartitionViewportWidget.cpp
@@ -74,7 +74,7 @@ EVisibility SWorldPartitionViewportWidget::GetVisibility(UWorld* InWorld)
 		UWorldPartition* WorldPartition = InWorld->GetWorldPartition();
 		if (WorldPartition->IsStreamingEnabled())
 		{
-			if (!WorldPartition->HasLoadedUserCreatedRegions())
+			if (!WorldPartition->HasLoadedUserCreatedRegions() && WorldPartition->IsLoadingInEditorEnabled())
 			{
 				Message = LOCTEXT("NoRegionsLoadedText","No Regions Loaded");
 				Tooltip = LOCTEXT("NoRegionsLoadedToolTip", "To load a region, drag select an area in the World Partition map and choose 'Load Region From Selection' from the context menu.");
diff --git a/Engine/Source/Editor/WorldPartitionEditor/Private/WorldPartition/SWorldPartitionEditorGrid2D.cpp b/Engine/Source/Editor/WorldPartitionEditor/Private/WorldPartition/SWorldPartitionEditorGrid2D.cpp
--- a/Engine/Source/Editor/WorldPartitionEditor/Private/WorldPartition/SWorldPartitionEditorGrid2D.cpp
+++ b/Engine/Source/Editor/WorldPartitionEditor/Private/WorldPartition/SWorldPartitionEditorGrid2D.cpp
@@ -714,19 +714,26 @@ FReply SWorldPartitionEditorGrid2D::OnMouseButtonUp(const FGeometry& MyGeometry,
 
 			const FEditorCommands& Commands = FEditorCommands::Get();
 
-			MenuBuilder.BeginSection(NAME_None, LOCTEXT("WorldPartitionSelection", "Selection"));
-				MenuBuilder.AddMenuEntry(Commands.CreateRegionFromSelection);
-				MenuBuilder.AddMenuSeparator();
-				MenuBuilder.AddMenuEntry(Commands.LoadSelectedRegions);
-				MenuBuilder.AddMenuEntry(Commands.UnloadSelectedRegions);
-				MenuBuilder.AddMenuSeparator();
-				MenuBuilder.AddMenuEntry(Commands.ConvertSelectedRegionsToActors);
-			MenuBuilder.EndSection();
+			if (WorldPartition->IsLoadingInEditorEnabled())
+			{
+				MenuBuilder.BeginSection(NAME_None, LOCTEXT("WorldPartitionSelection", "Selection"));
+					MenuBuilder.AddMenuEntry(Commands.CreateRegionFromSelection);
+					MenuBuilder.AddMenuSeparator();
+					MenuBuilder.AddMenuEntry(Commands.LoadSelectedRegions);
+					MenuBuilder.AddMenuEntry(Commands.UnloadSelectedRegions);
+					MenuBuilder.AddMenuSeparator();
+					MenuBuilder.AddMenuEntry(Commands.ConvertSelectedRegionsToActors);
+				MenuBuilder.EndSection();
+			}
 
 			MenuBuilder.BeginSection(NAME_None, LOCTEXT("WorldPartitionMisc", "Misc"));
 				MenuBuilder.AddMenuEntry(Commands.MoveCameraHere);
 				MenuBuilder.AddMenuEntry(Commands.PlayFromHere);
-				MenuBuilder.AddMenuEntry(Commands.LoadFromHere);
+				
+				if (WorldPartition->IsLoadingInEditorEnabled())
+				{
+					MenuBuilder.AddMenuEntry(Commands.LoadFromHere);
+				}
 			MenuBuilder.EndSection();			
 
 			FWidgetPath WidgetPath = MouseEvent.GetEventPath() != nullptr ? *MouseEvent.GetEventPath() : FWidgetPath();
@@ -803,7 +810,7 @@ FReply SWorldPartitionEditorGrid2D::OnMouseButtonDoubleClick(const FGeometry& In
 		{
 			MoveCameraHere();
 
-			if (InMouseEvent.IsControlDown())
+			if (InMouseEvent.IsControlDown() && WorldPartition->IsLoadingInEditorEnabled())
 			{
 				LoadFromHere();
 			}
diff --git a/Engine/Source/Runtime/Engine/Private/WorldPartition/WorldPartition.cpp b/Engine/Source/Runtime/Engine/Private/WorldPartition/WorldPartition.cpp
--- a/Engine/Source/Runtime/Engine/Private/WorldPartition/WorldPartition.cpp
+++ b/Engine/Source/Runtime/Engine/Private/WorldPartition/WorldPartition.cpp
@@ -244,6 +244,7 @@ UWorldPartition::UWorldPartition(const FObjectInitializer& ObjectInitializer)
 	, bStreamingWasEnabled(true)
 	, bShouldCheckEnableStreamingWarning(false)
 	, bCanBeUsedByLevelInstance(false)
+	, bEnableLoadingInEditor(true)
 	, bForceGarbageCollection(false)
 	, bForceGarbageCollectionPurge(false)
 	, bEnablingStreamingJustified(false)
@@ -713,7 +714,7 @@ bool UWorldPartition::IsMainWorldPartition() const
 void UWorldPartition::OnPostBugItGoCalled(const FVector& Loc, const FRotator& Rot)
 {
 #if WITH_EDITOR
-	if (GetMutableDefault<UWorldPartitionEditorPerProjectUserSettings>()->GetBugItGoLoadRegion())
+	if (GetMutableDefault<UWorldPartitionEditorPerProjectUserSettings>()->GetBugItGoLoadRegion() && IsLoadingInEditorEnabled())
 	{
 		const FVector LoadExtent(GLoadingRangeBugItGo, GLoadingRangeBugItGo, HALF_WORLD_MAX);
 		const FBox LoadCellsBox(Loc - LoadExtent, Loc + LoadExtent);
diff --git a/Engine/Source/Runtime/Engine/Private/WorldPartition/WorldPartitionConverter.cpp b/Engine/Source/Runtime/Engine/Private/WorldPartition/WorldPartitionConverter.cpp
--- a/Engine/Source/Runtime/Engine/Private/WorldPartition/WorldPartitionConverter.cpp
+++ b/Engine/Source/Runtime/Engine/Private/WorldPartition/WorldPartitionConverter.cpp
@@ -50,17 +50,25 @@ bool FWorldPartitionConverter::Convert()
 	UWorldPartition* WorldPartition = World->GetWorldPartition();
 	if (!WorldPartition)
 	{
+		if (Parameters.bCanBeUsedByLevelInstance && Parameters.bEnableStreaming)
+		{
+			return false;
+		}
+
 		World->PersistentLevel->ConvertAllActorsToPackaging(true);
 		World->PersistentLevel->bUseExternalActors = true;
+		
 		WorldPartition = UWorldPartition::CreateOrRepairWorldPartition(WorldSettings, Parameters.EditorHashClass, Parameters.RuntimeHashClass);
 		if (!WorldPartition)
 		{
 			return false;
 		}
+
 		FLevelActorFoldersHelper::SetUseActorFolders(World->PersistentLevel, Parameters.bUseActorFolders);
 		WorldPartition->bEnableStreaming = Parameters.bEnableStreaming;
 		WorldPartition->bStreamingWasEnabled = Parameters.bEnableStreaming;
-		WorldPartition->SetCanBeUsedByLevelInstance(!Parameters.bEnableStreaming);
+		WorldPartition->bEnableLoadingInEditor = Parameters.bEnableStreaming;
+		WorldPartition->SetCanBeUsedByLevelInstance(Parameters.bCanBeUsedByLevelInstance);
 		bCreatedWorldPartition = true;
 	}
 
diff --git a/Engine/Source/Runtime/Engine/Private/WorldPartition/WorldPartitionEditorSpatialHash.cpp b/Engine/Source/Runtime/Engine/Private/WorldPartition/WorldPartitionEditorSpatialHash.cpp
--- a/Engine/Source/Runtime/Engine/Private/WorldPartition/WorldPartitionEditorSpatialHash.cpp
+++ b/Engine/Source/Runtime/Engine/Private/WorldPartition/WorldPartitionEditorSpatialHash.cpp
@@ -115,14 +115,15 @@ void UWorldPartitionEditorSpatialHash::HashActor(FWorldPartitionHandle& InActorH
 {
 	check(InActorHandle.IsValid());
 
-	const FBox ActorBounds = InActorHandle->GetIsSpatiallyLoaded() ? InActorHandle->GetBounds() : FBox(ForceInit);
+	const bool bConsiderActorSpatiallyLoaded = InActorHandle->GetIsSpatiallyLoaded() && GetOuterUWorldPartition()->IsLoadingInEditorEnabled();
+	const FBox ActorBounds = bConsiderActorSpatiallyLoaded ? InActorHandle->GetBounds() : FBox(ForceInit);
 
 #if DO_CHECK
 	check(!HashedActors.Contains(InActorHandle->GetGuid()));
 	HashedActors.Add(InActorHandle->GetGuid(), ActorBounds);
 #endif
 
-	if (!InActorHandle->GetIsSpatiallyLoaded())
+	if (!bConsiderActorSpatiallyLoaded)
 	{
 		AlwaysLoadedCell->Actors.Add(InActorHandle);
 	}
@@ -225,14 +226,15 @@ void UWorldPartitionEditorSpatialHash::UnhashActor(FWorldPartitionHandle& InActo
 {
 	check(InActorHandle.IsValid());
 
-	const FBox ActorBounds = InActorHandle->GetIsSpatiallyLoaded() ? InActorHandle->GetBounds() : FBox(ForceInit);
+	const bool bConsiderActorSpatiallyLoaded = InActorHandle->GetIsSpatiallyLoaded() && GetOuterUWorldPartition()->IsLoadingInEditorEnabled();
+	const FBox ActorBounds = bConsiderActorSpatiallyLoaded ? InActorHandle->GetBounds() : FBox(ForceInit);
 
 #if DO_CHECK
-	FBox OldActorBounds = HashedActors.FindAndRemoveChecked(InActorHandle->GetGuid());
+	const FBox OldActorBounds = HashedActors.FindAndRemoveChecked(InActorHandle->GetGuid());
 	check(ActorBounds == OldActorBounds);
 #endif
 
-	if (!InActorHandle->GetIsSpatiallyLoaded())
+	if (!bConsiderActorSpatiallyLoaded)
 	{
 		AlwaysLoadedCell->Actors.Remove(InActorHandle);
 	}
diff --git a/Engine/Source/Runtime/Engine/Public/WorldPartition/WorldPartition.h b/Engine/Source/Runtime/Engine/Public/WorldPartition/WorldPartition.h
--- a/Engine/Source/Runtime/Engine/Public/WorldPartition/WorldPartition.h
+++ b/Engine/Source/Runtime/Engine/Public/WorldPartition/WorldPartition.h
@@ -202,6 +202,7 @@ public:
 	void OnUserCreatedRegionUnloaded() { check(HasLoadedUserCreatedRegions()); NumUserCreatedLoadedRegions--; }
 
 	bool IsEnablingStreamingJustified() const { return bEnablingStreamingJustified; }
+	bool IsLoadingInEditorEnabled() const { return bEnableLoadingInEditor; }
 #endif
 
 public:
@@ -261,6 +262,10 @@ private:
 	/** Whether Level Instance can reference this partition. */
 	UPROPERTY(EditAnywhere, Category = "WorldPartitionSetup", AdvancedDisplay, meta = (EditConditionHides, HideEditConditionToggle, EditCondition = "!bEnableStreaming"))
 	bool bCanBeUsedByLevelInstance;
+
+	/** When disabled, editor will always load the whole world. */
+	UPROPERTY()
+	bool bEnableLoadingInEditor;
 #endif
 
 public:
diff --git a/Engine/Source/Runtime/Engine/Public/WorldPartition/WorldPartitionConverter.h b/Engine/Source/Runtime/Engine/Public/WorldPartition/WorldPartitionConverter.h
--- a/Engine/Source/Runtime/Engine/Public/WorldPartition/WorldPartitionConverter.h
+++ b/Engine/Source/Runtime/Engine/Public/WorldPartition/WorldPartitionConverter.h
@@ -21,6 +21,8 @@ public:
 		FParameters()
 		: bConvertSubLevels(true)
 		, bEnableStreaming(true)
+		, bEnableLoadingInEditor(true)
+		, bCanBeUsedByLevelInstance(false)
 		, bUseActorFolders(true)
 		, EditorHashClass(nullptr)
 		, RuntimeHashClass(nullptr)
@@ -28,6 +30,8 @@ public:
 
 		bool bConvertSubLevels;
 		bool bEnableStreaming;
+		bool bEnableLoadingInEditor;
+		bool bCanBeUsedByLevelInstance;
 		bool bUseActorFolders;
 		TSubclassOf<UWorldPartitionEditorHash> EditorHashClass;
 		TSubclassOf<UWorldPartitionRuntimeHash> RuntimeHashClass;
