World Partition: added editor settings for PIE and loading regions for projects that don't support these.

#rnx

diff --git a/Engine/Source/Editor/UnrealEd/Private/SWorldPartitionViewportWidget.cpp b/Engine/Source/Editor/UnrealEd/Private/SWorldPartitionViewportWidget.cpp
--- a/Engine/Source/Editor/UnrealEd/Private/SWorldPartitionViewportWidget.cpp
+++ b/Engine/Source/Editor/UnrealEd/Private/SWorldPartitionViewportWidget.cpp
@@ -10,6 +10,7 @@
 #include "Editor.h"
 #include "Engine/World.h"
 #include "WorldPartition/WorldPartition.h"
+#include "WorldPartition/IWorldPartitionEditorModule.h"
 #include "Styling/AppStyle.h"
 #include "Modules/ModuleManager.h"
 
@@ -62,6 +63,8 @@ void SWorldPartitionViewportWidget::Construct(const FArguments& InArgs)
 			]
 		]
 	];
+
+	WorldPartitionEditorModule = FModuleManager::GetModulePtr<IWorldPartitionEditorModule>("WorldPartitionEditor");
 }
 
 SWorldPartitionViewportWidget::~SWorldPartitionViewportWidget()
@@ -74,7 +77,7 @@ EVisibility SWorldPartitionViewportWidget::GetVisibility(UWorld* InWorld)
 		UWorldPartition* WorldPartition = InWorld->GetWorldPartition();
 		if (WorldPartition->IsStreamingEnabled())
 		{
-			if (!WorldPartition->HasLoadedUserCreatedRegions() && WorldPartition->IsLoadingInEditorEnabled())
+			if (!WorldPartition->HasLoadedUserCreatedRegions() && !WorldPartitionEditorModule->GetDisableLoadingInEditor())
 			{
 				Message = LOCTEXT("NoRegionsLoadedText","No Regions Loaded");
 				Tooltip = LOCTEXT("NoRegionsLoadedToolTip", "To load a region, drag select an area in the World Partition map and choose 'Load Region From Selection' from the context menu.");
diff --git a/Engine/Source/Editor/UnrealEd/Public/SWorldPartitionViewportWidget.h b/Engine/Source/Editor/UnrealEd/Public/SWorldPartitionViewportWidget.h
--- a/Engine/Source/Editor/UnrealEd/Public/SWorldPartitionViewportWidget.h
+++ b/Engine/Source/Editor/UnrealEd/Public/SWorldPartitionViewportWidget.h
@@ -4,6 +4,8 @@
 
 #include "Widgets/SCompoundWidget.h"
 
+class IWorldPartitionEditorModule;
+
 class UNREALED_API SWorldPartitionViewportWidget : public SCompoundWidget
 {
 public:
@@ -21,4 +23,5 @@ private:
 	FText Message;
 	FText Tooltip;
 	FName InvokeTab;
+	IWorldPartitionEditorModule* WorldPartitionEditorModule;
 };
\ No newline at end of file
diff --git a/Engine/Source/Editor/WorldPartitionEditor/Private/WorldPartition/SWorldPartitionEditorGrid2D.cpp b/Engine/Source/Editor/WorldPartitionEditor/Private/WorldPartition/SWorldPartitionEditorGrid2D.cpp
--- a/Engine/Source/Editor/WorldPartitionEditor/Private/WorldPartition/SWorldPartitionEditorGrid2D.cpp
+++ b/Engine/Source/Editor/WorldPartitionEditor/Private/WorldPartition/SWorldPartitionEditorGrid2D.cpp
@@ -12,6 +12,7 @@
 #include "WorldPartition/WorldPartitionActorDesc.h"
 #include "WorldPartition/WorldPartitionActorDescView.h"
 #include "WorldPartition/WorldPartitionEditorPerProjectUserSettings.h"
+#include "WorldPartition/WorldPartitionEditorSettings.h"
 #include "WorldPartition/WorldPartitionHelpers.h"
 #include "WorldPartition/WorldPartitionMiniMap.h"
 #include "WorldPartition/WorldPartitionMiniMapHelper.h"
@@ -321,6 +322,7 @@ void SWorldPartitionEditorGrid2D::Construct(const FArguments& InArgs)
 						SNew(SCheckBox)
 						.IsChecked(bFollowPlayerInPIE ? ECheckBoxState::Checked : ECheckBoxState::Unchecked)
 						.IsEnabled(true)
+						.Visibility_Lambda([this]() { return GetDefault<UWorldPartitionEditorSettings>()->bDisablePIE ? EVisibility::Hidden : EVisibility::Visible; })
 						.OnCheckStateChanged(FOnCheckStateChanged::CreateLambda([=](ECheckBoxState State) { bFollowPlayerInPIE = !bFollowPlayerInPIE; }))
 					]
 					+SHorizontalBox::Slot()
@@ -330,6 +332,7 @@ void SWorldPartitionEditorGrid2D::Construct(const FArguments& InArgs)
 						SNew(STextBlock)
 						.AutoWrapText(true)
 						.IsEnabled(true)
+						.Visibility_Lambda([this]() { return GetDefault<UWorldPartitionEditorSettings>()->bDisablePIE ? EVisibility::Hidden : EVisibility::Visible; })
 						.Text(LOCTEXT("FollowPlayerInPIE", "Follow Player in PIE"))
 					]
 					+SHorizontalBox::Slot()
@@ -338,6 +341,7 @@ void SWorldPartitionEditorGrid2D::Construct(const FArguments& InArgs)
 						SNew(SCheckBox)
 						.IsChecked(GetMutableDefault<UWorldPartitionEditorPerProjectUserSettings>()->GetBugItGoLoadRegion() ? ECheckBoxState::Checked : ECheckBoxState::Unchecked)
 						.IsEnabled(true)
+						.Visibility_Lambda([this]() { return GetDefault<UWorldPartitionEditorSettings>()->bDisableLoadingInEditor ? EVisibility::Hidden : EVisibility::Visible; })
 						.OnCheckStateChanged(FOnCheckStateChanged::CreateLambda([=](ECheckBoxState State) { GetMutableDefault<UWorldPartitionEditorPerProjectUserSettings>()->SetBugItGoLoadRegion(State == ECheckBoxState::Checked); }))
 					]
 					+SHorizontalBox::Slot()
@@ -347,6 +351,7 @@ void SWorldPartitionEditorGrid2D::Construct(const FArguments& InArgs)
 						SNew(STextBlock)
 						.AutoWrapText(true)
 						.IsEnabled(true)
+						.Visibility_Lambda([this]() { return GetDefault<UWorldPartitionEditorSettings>()->bDisableLoadingInEditor ? EVisibility::Hidden : EVisibility::Visible; })
 						.Text(LOCTEXT("BugItGoLoadRegion", "BugItGo Load Region"))
 					]
 					+SHorizontalBox::Slot()
@@ -381,6 +386,7 @@ void SWorldPartitionEditorGrid2D::Construct(const FArguments& InArgs)
 						.Text(LOCTEXT("FocusLoadedRegions", "Focus Loaded Regions"))
 						.OnClicked(this, &SWorldPartitionEditorGrid2D::FocusLoadedRegions)
 						.IsEnabled_Lambda([this]() { return IsInteractive() && WorldPartition && WorldPartition->HasLoadedUserCreatedRegions(); })
+						.Visibility_Lambda([this]() { return GetDefault<UWorldPartitionEditorSettings>()->bDisableLoadingInEditor ? EVisibility::Hidden : EVisibility::Visible; })
 					]
 				]
 			]
@@ -714,7 +720,7 @@ FReply SWorldPartitionEditorGrid2D::OnMouseButtonUp(const FGeometry& MyGeometry,
 
 			const FEditorCommands& Commands = FEditorCommands::Get();
 
-			if (WorldPartition->IsLoadingInEditorEnabled())
+			if (!GetDefault<UWorldPartitionEditorSettings>()->bDisableLoadingInEditor)
 			{
 				MenuBuilder.BeginSection(NAME_None, LOCTEXT("WorldPartitionSelection", "Selection"));
 					MenuBuilder.AddMenuEntry(Commands.CreateRegionFromSelection);
@@ -728,9 +734,13 @@ FReply SWorldPartitionEditorGrid2D::OnMouseButtonUp(const FGeometry& MyGeometry,
 
 			MenuBuilder.BeginSection(NAME_None, LOCTEXT("WorldPartitionMisc", "Misc"));
 				MenuBuilder.AddMenuEntry(Commands.MoveCameraHere);
-				MenuBuilder.AddMenuEntry(Commands.PlayFromHere);
+
+				if (!GetDefault<UWorldPartitionEditorSettings>()->bDisablePIE)
+				{
+					MenuBuilder.AddMenuEntry(Commands.PlayFromHere);
+				}
 				
-				if (WorldPartition->IsLoadingInEditorEnabled())
+				if (!GetDefault<UWorldPartitionEditorSettings>()->bDisableLoadingInEditor)
 				{
 					MenuBuilder.AddMenuEntry(Commands.LoadFromHere);
 				}
@@ -804,13 +814,16 @@ FReply SWorldPartitionEditorGrid2D::OnMouseButtonDoubleClick(const FGeometry& In
 	{
 		if (InMouseEvent.IsShiftDown())
 		{
-			PlayFromHere();
+			if (!GetDefault<UWorldPartitionEditorSettings>()->bDisablePIE)
+			{
+				PlayFromHere();
+			}
 		}
 		else
 		{
 			MoveCameraHere();
 
-			if (InMouseEvent.IsControlDown() && WorldPartition->IsLoadingInEditorEnabled())
+			if (InMouseEvent.IsControlDown() && !GetDefault<UWorldPartitionEditorSettings>()->bDisableLoadingInEditor)
 			{
 				LoadFromHere();
 			}
diff --git a/Engine/Source/Editor/WorldPartitionEditor/Private/WorldPartitionEditorModule.cpp b/Engine/Source/Editor/WorldPartitionEditor/Private/WorldPartitionEditorModule.cpp
--- a/Engine/Source/Editor/WorldPartitionEditor/Private/WorldPartitionEditorModule.cpp
+++ b/Engine/Source/Editor/WorldPartitionEditor/Private/WorldPartitionEditorModule.cpp
@@ -240,6 +240,26 @@ int32 FWorldPartitionEditorModule::GetMinimapLowQualityWorldUnitsPerPixelThresho
 	return GetDefault<UWorldPartitionEditorSettings>()->MinimapLowQualityWorldUnitsPerPixelThreshold;
 }
 
+bool FWorldPartitionEditorModule::GetDisableLoadingInEditor() const
+{
+	return GetDefault<UWorldPartitionEditorSettings>()->bDisableLoadingInEditor;
+}
+
+void FWorldPartitionEditorModule::SetDisableLoadingInEditor(bool bInDisableLoadingInEditor)
+{
+	GetMutableDefault<UWorldPartitionEditorSettings>()->bDisableLoadingInEditor = bInDisableLoadingInEditor;
+}
+
+bool FWorldPartitionEditorModule::GetDisablePIE() const
+{
+	return GetDefault<UWorldPartitionEditorSettings>()->bDisablePIE;
+}
+
+void FWorldPartitionEditorModule::SetDisablePIE(bool bInDisablePIE)
+{
+	GetMutableDefault<UWorldPartitionEditorSettings>()->bDisablePIE = bInDisablePIE;
+}
+
 void FWorldPartitionEditorModule::OnConvertMap()
 {
 	IContentBrowserSingleton& ContentBrowserSingleton = FModuleManager::LoadModuleChecked<FContentBrowserModule>("ContentBrowser").Get();
@@ -676,6 +696,8 @@ UWorldPartitionEditorSettings::UWorldPartitionEditorSettings()
 	CommandletClass = UWorldPartitionConvertCommandlet::StaticClass();
 	InstancedFoliageGridSize = 25600;
 	MinimapLowQualityWorldUnitsPerPixelThreshold = 12800;
+	bDisableLoadingInEditor = false;
+	bDisablePIE = false;
 }
 
 FString UWorldPartitionConvertOptions::ToCommandletArgs() const
diff --git a/Engine/Source/Editor/WorldPartitionEditor/Public/WorldPartition/WorldPartitionEditorSettings.h b/Engine/Source/Editor/WorldPartitionEditor/Public/WorldPartition/WorldPartitionEditorSettings.h
--- a/Engine/Source/Editor/WorldPartitionEditor/Public/WorldPartition/WorldPartitionEditorSettings.h
+++ b/Engine/Source/Editor/WorldPartitionEditor/Public/WorldPartition/WorldPartitionEditorSettings.h
@@ -15,12 +15,15 @@ class UWorldPartitionEditorSettings : public UDeveloperSettings
 public:
 	UWorldPartitionEditorSettings();
 
-	UPROPERTY(config, EditAnywhere, Category = MapConversion, meta = (ToolTip = "Commandlet class to use for World Parition conversion"))
+	UPROPERTY(Config, EditAnywhere, Category = MapConversion, Meta = (ToolTip = "Commandlet class to use for World Parition conversion"))
 	TSubclassOf<UWorldPartitionConvertCommandlet> CommandletClass;
 
-	UPROPERTY(config, EditAnywhere, Category = Foliage, meta = (ClampMin=3200, ToolTip= "Editor grid size used for instance foliage actors in World Partition worlds"))
+	UPROPERTY(Config, EditAnywhere, Category = Foliage, Meta = (ClampMin=3200, ToolTip= "Editor grid size used for instance foliage actors in World Partition worlds"))
 	int32 InstancedFoliageGridSize;
 
-	UPROPERTY(config, EditAnywhere, Category = MiniMap, meta = (ClampMin = 100, ToolTip = "Threshold from which minimap generates a warning if its WorldUnitsPerPixel is above this value"))
+	UPROPERTY(Config, EditAnywhere, Category = MiniMap, Meta = (ClampMin = 100, ToolTip = "Threshold from which minimap generates a warning if its WorldUnitsPerPixel is above this value"))
 	int32 MinimapLowQualityWorldUnitsPerPixelThreshold;
+
+	bool bDisableLoadingInEditor;
+	bool bDisablePIE;
 };
\ No newline at end of file
diff --git a/Engine/Source/Editor/WorldPartitionEditor/Public/WorldPartitionEditorModule.h b/Engine/Source/Editor/WorldPartitionEditor/Public/WorldPartitionEditorModule.h
--- a/Engine/Source/Editor/WorldPartitionEditor/Public/WorldPartitionEditorModule.h
+++ b/Engine/Source/Editor/WorldPartitionEditor/Public/WorldPartitionEditorModule.h
@@ -47,6 +47,18 @@ public:
 	 */
 	virtual int32 GetMinimapLowQualityWorldUnitsPerPixelThreshold() const override;
 
+	/**
+	 * Returns if loading in the editor is disabled or not.
+	 */
+	virtual bool GetDisableLoadingInEditor() const override;
+	virtual void SetDisableLoadingInEditor(bool bInDisableLoadingInEditor) override;
+
+	/**
+	 * Returns if pie is disabled or not..
+	 */
+	virtual bool GetDisablePIE() const override;
+	virtual void SetDisablePIE(bool bInDisablePIE) override;
+
 	/**
 	 * Convert the specified map to a world partition map.
 	 */
diff --git a/Engine/Source/Runtime/Engine/Private/WorldPartition/WorldPartition.cpp b/Engine/Source/Runtime/Engine/Private/WorldPartition/WorldPartition.cpp
--- a/Engine/Source/Runtime/Engine/Private/WorldPartition/WorldPartition.cpp
+++ b/Engine/Source/Runtime/Engine/Private/WorldPartition/WorldPartition.cpp
@@ -244,7 +244,6 @@ UWorldPartition::UWorldPartition(const FObjectInitializer& ObjectInitializer)
 	, bStreamingWasEnabled(true)
 	, bShouldCheckEnableStreamingWarning(false)
 	, bCanBeUsedByLevelInstance(false)
-	, bEnableLoadingInEditor(true)
 	, bForceGarbageCollection(false)
 	, bForceGarbageCollectionPurge(false)
 	, bEnablingStreamingJustified(false)
@@ -714,13 +713,17 @@ bool UWorldPartition::IsMainWorldPartition() const
 void UWorldPartition::OnPostBugItGoCalled(const FVector& Loc, const FRotator& Rot)
 {
 #if WITH_EDITOR
-	if (GetMutableDefault<UWorldPartitionEditorPerProjectUserSettings>()->GetBugItGoLoadRegion() && IsLoadingInEditorEnabled())
+	if (GetMutableDefault<UWorldPartitionEditorPerProjectUserSettings>()->GetBugItGoLoadRegion())
 	{
-		const FVector LoadExtent(GLoadingRangeBugItGo, GLoadingRangeBugItGo, HALF_WORLD_MAX);
-		const FBox LoadCellsBox(Loc - LoadExtent, Loc + LoadExtent);
+		IWorldPartitionEditorModule& WorldPartitionEditorModule = FModuleManager::LoadModuleChecked<IWorldPartitionEditorModule>("WorldPartitionEditor");
+		if (!WorldPartitionEditorModule.GetDisableLoadingInEditor())
+		{
+			const FVector LoadExtent(GLoadingRangeBugItGo, GLoadingRangeBugItGo, HALF_WORLD_MAX);
+			const FBox LoadCellsBox(Loc - LoadExtent, Loc + LoadExtent);
 
-		UWorldPartitionEditorLoaderAdapter* EditorLoaderAdapter = CreateEditorLoaderAdapter<FLoaderAdapterShape>(World, LoadCellsBox, TEXT("BugItGo"));
-		EditorLoaderAdapter->GetLoaderAdapter()->Load();
+			UWorldPartitionEditorLoaderAdapter* EditorLoaderAdapter = CreateEditorLoaderAdapter<FLoaderAdapterShape>(World, LoadCellsBox, TEXT("BugItGo"));
+			EditorLoaderAdapter->GetLoaderAdapter()->Load();
+		}
 	}
 #endif
 }
diff --git a/Engine/Source/Runtime/Engine/Private/WorldPartition/WorldPartitionConverter.cpp b/Engine/Source/Runtime/Engine/Private/WorldPartition/WorldPartitionConverter.cpp
--- a/Engine/Source/Runtime/Engine/Private/WorldPartition/WorldPartitionConverter.cpp
+++ b/Engine/Source/Runtime/Engine/Private/WorldPartition/WorldPartitionConverter.cpp
@@ -67,7 +67,6 @@ bool FWorldPartitionConverter::Convert()
 		FLevelActorFoldersHelper::SetUseActorFolders(World->PersistentLevel, Parameters.bUseActorFolders);
 		WorldPartition->bEnableStreaming = Parameters.bEnableStreaming;
 		WorldPartition->bStreamingWasEnabled = Parameters.bEnableStreaming;
-		WorldPartition->bEnableLoadingInEditor = Parameters.bEnableStreaming;
 		WorldPartition->SetCanBeUsedByLevelInstance(Parameters.bCanBeUsedByLevelInstance);
 		bCreatedWorldPartition = true;
 	}
diff --git a/Engine/Source/Runtime/Engine/Private/WorldPartition/WorldPartitionEditorHash.cpp b/Engine/Source/Runtime/Engine/Private/WorldPartition/WorldPartitionEditorHash.cpp
--- a/Engine/Source/Runtime/Engine/Private/WorldPartition/WorldPartitionEditorHash.cpp
+++ b/Engine/Source/Runtime/Engine/Private/WorldPartition/WorldPartitionEditorHash.cpp
@@ -1,12 +1,17 @@
 // Copyright Epic Games, Inc. All Rights Reserved.
 
 #include "WorldPartition/WorldPartitionEditorHash.h"
+#include "WorldPartition/IWorldPartitionEditorModule.h"
 
 #include UE_INLINE_GENERATED_CPP_BY_NAME(WorldPartitionEditorHash)
 
 UWorldPartitionEditorHash::UWorldPartitionEditorHash(const FObjectInitializer& ObjectInitializer)
 	: Super(ObjectInitializer)
-{}
+{
+#if WITH_EDITOR
+	WorldPartitionEditorModule = &FModuleManager::LoadModuleChecked<IWorldPartitionEditorModule>("WorldPartitionEditor");
+#endif
+}
 
 #if WITH_EDITOR
 UWorldPartitionEditorHash::FForEachIntersectingActorParams::FForEachIntersectingActorParams()
diff --git a/Engine/Source/Runtime/Engine/Private/WorldPartition/WorldPartitionEditorSpatialHash.cpp b/Engine/Source/Runtime/Engine/Private/WorldPartition/WorldPartitionEditorSpatialHash.cpp
--- a/Engine/Source/Runtime/Engine/Private/WorldPartition/WorldPartitionEditorSpatialHash.cpp
+++ b/Engine/Source/Runtime/Engine/Private/WorldPartition/WorldPartitionEditorSpatialHash.cpp
@@ -5,7 +5,7 @@
 #include "WorldPartition/WorldPartitionLog.h"
 #include "WorldPartition/WorldPartitionActorDesc.h"
 #include "WorldPartition/WorldPartitionActorDesc.h"
-#include "WorldPartition/WorldPartitionEditorPerProjectUserSettings.h"
+#include "WorldPartition/IWorldPartitionEditorModule.h"
 
 #include UE_INLINE_GENERATED_CPP_BY_NAME(WorldPartitionEditorSpatialHash)
 
@@ -115,7 +115,7 @@ void UWorldPartitionEditorSpatialHash::HashActor(FWorldPartitionHandle& InActorH
 {
 	check(InActorHandle.IsValid());
 
-	const bool bConsiderActorSpatiallyLoaded = InActorHandle->GetIsSpatiallyLoaded() && GetOuterUWorldPartition()->IsLoadingInEditorEnabled();
+	const bool bConsiderActorSpatiallyLoaded = InActorHandle->GetIsSpatiallyLoaded() && !WorldPartitionEditorModule->GetDisableLoadingInEditor();
 	const FBox ActorBounds = bConsiderActorSpatiallyLoaded ? InActorHandle->GetBounds() : FBox(ForceInit);
 
 #if DO_CHECK
@@ -226,7 +226,7 @@ void UWorldPartitionEditorSpatialHash::UnhashActor(FWorldPartitionHandle& InActo
 {
 	check(InActorHandle.IsValid());
 
-	const bool bConsiderActorSpatiallyLoaded = InActorHandle->GetIsSpatiallyLoaded() && GetOuterUWorldPartition()->IsLoadingInEditorEnabled();
+	const bool bConsiderActorSpatiallyLoaded = InActorHandle->GetIsSpatiallyLoaded()  && !WorldPartitionEditorModule->GetDisableLoadingInEditor();
 	const FBox ActorBounds = bConsiderActorSpatiallyLoaded ? InActorHandle->GetBounds() : FBox(ForceInit);
 
 #if DO_CHECK
diff --git a/Engine/Source/Runtime/Engine/Public/WorldPartition/IWorldPartitionEditorModule.h b/Engine/Source/Runtime/Engine/Public/WorldPartition/IWorldPartitionEditorModule.h
--- a/Engine/Source/Runtime/Engine/Public/WorldPartition/IWorldPartitionEditorModule.h
+++ b/Engine/Source/Runtime/Engine/Public/WorldPartition/IWorldPartitionEditorModule.h
@@ -34,6 +34,12 @@ public:
 	virtual int32 GetPlacementGridSize() const = 0;
 	virtual int32 GetInstancedFoliageGridSize() const = 0;
 	virtual int32 GetMinimapLowQualityWorldUnitsPerPixelThreshold() const = 0;
+	
+	virtual bool GetDisableLoadingInEditor() const = 0;
+	virtual void SetDisableLoadingInEditor(bool bInDisableLoadingInEditor) = 0;
+
+	virtual bool GetDisablePIE() const = 0;
+	virtual void SetDisablePIE(bool bInDisablePIE) = 0;
 
 	/** Triggered when a world is added. */
 	DECLARE_EVENT_OneParam(IWorldPartitionEditorModule, FWorldPartitionCreated, UWorld*);
diff --git a/Engine/Source/Runtime/Engine/Public/WorldPartition/WorldPartition.h b/Engine/Source/Runtime/Engine/Public/WorldPartition/WorldPartition.h
--- a/Engine/Source/Runtime/Engine/Public/WorldPartition/WorldPartition.h
+++ b/Engine/Source/Runtime/Engine/Public/WorldPartition/WorldPartition.h
@@ -202,7 +202,6 @@ public:
 	void OnUserCreatedRegionUnloaded() { check(HasLoadedUserCreatedRegions()); NumUserCreatedLoadedRegions--; }
 
 	bool IsEnablingStreamingJustified() const { return bEnablingStreamingJustified; }
-	bool IsLoadingInEditorEnabled() const { return bEnableLoadingInEditor; }
 #endif
 
 public:
@@ -262,10 +261,6 @@ private:
 	/** Whether Level Instance can reference this partition. */
 	UPROPERTY(EditAnywhere, Category = "WorldPartitionSetup", AdvancedDisplay, meta = (EditConditionHides, HideEditConditionToggle, EditCondition = "!bEnableStreaming"))
 	bool bCanBeUsedByLevelInstance;
-
-	/** When disabled, editor will always load the whole world. */
-	UPROPERTY()
-	bool bEnableLoadingInEditor;
 #endif
 
 public:
diff --git a/Engine/Source/Runtime/Engine/Public/WorldPartition/WorldPartitionEditorHash.h b/Engine/Source/Runtime/Engine/Public/WorldPartition/WorldPartitionEditorHash.h
--- a/Engine/Source/Runtime/Engine/Public/WorldPartition/WorldPartitionEditorHash.h
+++ b/Engine/Source/Runtime/Engine/Public/WorldPartition/WorldPartitionEditorHash.h
@@ -8,6 +8,8 @@
 #include "WorldPartition/WorldPartitionHandle.h"
 #include "WorldPartitionEditorHash.generated.h"
 
+class IWorldPartitionEditorModule;
+
 UCLASS(Abstract, Config=Engine, Within = WorldPartition)
 class ENGINE_API UWorldPartitionEditorHash : public UObject
 {
@@ -44,5 +46,8 @@ public:
 	};
 
 	virtual int32 ForEachIntersectingActor(const FBox& Box, TFunctionRef<void(FWorldPartitionActorDesc*)> InOperation, const FForEachIntersectingActorParams& Params = FForEachIntersectingActorParams()) PURE_VIRTUAL(UWorldPartitionEditorHash::ForEachIntersectingActor, return 0;);
+
+protected:
+	IWorldPartitionEditorModule* WorldPartitionEditorModule;
 #endif
 };
