[Backout] - CL22220317
Original CL Desc
-----------------------------------------------------------------
World Partition: added editor settings for PIE and loading regions for projects that don't support these.

#rnx

diff --git a/Engine/Source/Editor/UnrealEd/Private/SWorldPartitionViewportWidget.cpp b/Engine/Source/Editor/UnrealEd/Private/SWorldPartitionViewportWidget.cpp
--- a/Engine/Source/Editor/UnrealEd/Private/SWorldPartitionViewportWidget.cpp
+++ b/Engine/Source/Editor/UnrealEd/Private/SWorldPartitionViewportWidget.cpp
@@ -10,7 +10,6 @@
 #include "Editor.h"
 #include "Engine/World.h"
 #include "WorldPartition/WorldPartition.h"
-#include "WorldPartition/IWorldPartitionEditorModule.h"
 #include "Styling/AppStyle.h"
 #include "Modules/ModuleManager.h"
 
@@ -63,8 +62,6 @@ void SWorldPartitionViewportWidget::Construct(const FArguments& InArgs)
 			]
 		]
 	];
-
-	WorldPartitionEditorModule = FModuleManager::GetModulePtr<IWorldPartitionEditorModule>("WorldPartitionEditor");
 }
 
 SWorldPartitionViewportWidget::~SWorldPartitionViewportWidget()
@@ -77,7 +74,7 @@ EVisibility SWorldPartitionViewportWidget::GetVisibility(UWorld* InWorld)
 		UWorldPartition* WorldPartition = InWorld->GetWorldPartition();
 		if (WorldPartition->IsStreamingEnabled())
 		{
-			if (!WorldPartition->HasLoadedUserCreatedRegions() && !WorldPartitionEditorModule->GetDisableLoadingInEditor())
+			if (!WorldPartition->HasLoadedUserCreatedRegions() && WorldPartition->IsLoadingInEditorEnabled())
 			{
 				Message = LOCTEXT("NoRegionsLoadedText","No Regions Loaded");
 				Tooltip = LOCTEXT("NoRegionsLoadedToolTip", "To load a region, drag select an area in the World Partition map and choose 'Load Region From Selection' from the context menu.");
diff --git a/Engine/Source/Editor/UnrealEd/Public/SWorldPartitionViewportWidget.h b/Engine/Source/Editor/UnrealEd/Public/SWorldPartitionViewportWidget.h
--- a/Engine/Source/Editor/UnrealEd/Public/SWorldPartitionViewportWidget.h
+++ b/Engine/Source/Editor/UnrealEd/Public/SWorldPartitionViewportWidget.h
@@ -4,8 +4,6 @@
 
 #include "Widgets/SCompoundWidget.h"
 
-class IWorldPartitionEditorModule;
-
 class UNREALED_API SWorldPartitionViewportWidget : public SCompoundWidget
 {
 public:
@@ -23,5 +21,4 @@ private:
 	FText Message;
 	FText Tooltip;
 	FName InvokeTab;
-	IWorldPartitionEditorModule* WorldPartitionEditorModule;
 };
\ No newline at end of file
diff --git a/Engine/Source/Editor/WorldPartitionEditor/Private/WorldPartition/SWorldPartitionEditorGrid2D.cpp b/Engine/Source/Editor/WorldPartitionEditor/Private/WorldPartition/SWorldPartitionEditorGrid2D.cpp
--- a/Engine/Source/Editor/WorldPartitionEditor/Private/WorldPartition/SWorldPartitionEditorGrid2D.cpp
+++ b/Engine/Source/Editor/WorldPartitionEditor/Private/WorldPartition/SWorldPartitionEditorGrid2D.cpp
@@ -12,7 +12,6 @@
 #include "WorldPartition/WorldPartitionActorDesc.h"
 #include "WorldPartition/WorldPartitionActorDescView.h"
 #include "WorldPartition/WorldPartitionEditorPerProjectUserSettings.h"
-#include "WorldPartition/WorldPartitionEditorSettings.h"
 #include "WorldPartition/WorldPartitionHelpers.h"
 #include "WorldPartition/WorldPartitionMiniMap.h"
 #include "WorldPartition/WorldPartitionMiniMapHelper.h"
@@ -322,7 +321,6 @@ void SWorldPartitionEditorGrid2D::Construct(const FArguments& InArgs)
 						SNew(SCheckBox)
 						.IsChecked(bFollowPlayerInPIE ? ECheckBoxState::Checked : ECheckBoxState::Unchecked)
 						.IsEnabled(true)
-						.Visibility_Lambda([this]() { return GetDefault<UWorldPartitionEditorSettings>()->bDisablePIE ? EVisibility::Hidden : EVisibility::Visible; })
 						.OnCheckStateChanged(FOnCheckStateChanged::CreateLambda([=](ECheckBoxState State) { bFollowPlayerInPIE = !bFollowPlayerInPIE; }))
 					]
 					+SHorizontalBox::Slot()
@@ -332,7 +330,6 @@ void SWorldPartitionEditorGrid2D::Construct(const FArguments& InArgs)
 						SNew(STextBlock)
 						.AutoWrapText(true)
 						.IsEnabled(true)
-						.Visibility_Lambda([this]() { return GetDefault<UWorldPartitionEditorSettings>()->bDisablePIE ? EVisibility::Hidden : EVisibility::Visible; })
 						.Text(LOCTEXT("FollowPlayerInPIE", "Follow Player in PIE"))
 					]
 					+SHorizontalBox::Slot()
@@ -341,7 +338,6 @@ void SWorldPartitionEditorGrid2D::Construct(const FArguments& InArgs)
 						SNew(SCheckBox)
 						.IsChecked(GetMutableDefault<UWorldPartitionEditorPerProjectUserSettings>()->GetBugItGoLoadRegion() ? ECheckBoxState::Checked : ECheckBoxState::Unchecked)
 						.IsEnabled(true)
-						.Visibility_Lambda([this]() { return GetDefault<UWorldPartitionEditorSettings>()->bDisableLoadingInEditor ? EVisibility::Hidden : EVisibility::Visible; })
 						.OnCheckStateChanged(FOnCheckStateChanged::CreateLambda([=](ECheckBoxState State) { GetMutableDefault<UWorldPartitionEditorPerProjectUserSettings>()->SetBugItGoLoadRegion(State == ECheckBoxState::Checked); }))
 					]
 					+SHorizontalBox::Slot()
@@ -351,7 +347,6 @@ void SWorldPartitionEditorGrid2D::Construct(const FArguments& InArgs)
 						SNew(STextBlock)
 						.AutoWrapText(true)
 						.IsEnabled(true)
-						.Visibility_Lambda([this]() { return GetDefault<UWorldPartitionEditorSettings>()->bDisableLoadingInEditor ? EVisibility::Hidden : EVisibility::Visible; })
 						.Text(LOCTEXT("BugItGoLoadRegion", "BugItGo Load Region"))
 					]
 					+SHorizontalBox::Slot()
@@ -386,7 +381,6 @@ void SWorldPartitionEditorGrid2D::Construct(const FArguments& InArgs)
 						.Text(LOCTEXT("FocusLoadedRegions", "Focus Loaded Regions"))
 						.OnClicked(this, &SWorldPartitionEditorGrid2D::FocusLoadedRegions)
 						.IsEnabled_Lambda([this]() { return IsInteractive() && WorldPartition && WorldPartition->HasLoadedUserCreatedRegions(); })
-						.Visibility_Lambda([this]() { return GetDefault<UWorldPartitionEditorSettings>()->bDisableLoadingInEditor ? EVisibility::Hidden : EVisibility::Visible; })
 					]
 				]
 			]
@@ -720,7 +714,7 @@ FReply SWorldPartitionEditorGrid2D::OnMouseButtonUp(const FGeometry& MyGeometry,
 
 			const FEditorCommands& Commands = FEditorCommands::Get();
 
-			if (!GetDefault<UWorldPartitionEditorSettings>()->bDisableLoadingInEditor)
+			if (WorldPartition->IsLoadingInEditorEnabled())
 			{
 				MenuBuilder.BeginSection(NAME_None, LOCTEXT("WorldPartitionSelection", "Selection"));
 					MenuBuilder.AddMenuEntry(Commands.CreateRegionFromSelection);
@@ -734,13 +728,9 @@ FReply SWorldPartitionEditorGrid2D::OnMouseButtonUp(const FGeometry& MyGeometry,
 
 			MenuBuilder.BeginSection(NAME_None, LOCTEXT("WorldPartitionMisc", "Misc"));
 				MenuBuilder.AddMenuEntry(Commands.MoveCameraHere);
-
-				if (!GetDefault<UWorldPartitionEditorSettings>()->bDisablePIE)
-				{
-					MenuBuilder.AddMenuEntry(Commands.PlayFromHere);
-				}
+				MenuBuilder.AddMenuEntry(Commands.PlayFromHere);
 				
-				if (!GetDefault<UWorldPartitionEditorSettings>()->bDisableLoadingInEditor)
+				if (WorldPartition->IsLoadingInEditorEnabled())
 				{
 					MenuBuilder.AddMenuEntry(Commands.LoadFromHere);
 				}
@@ -814,16 +804,13 @@ FReply SWorldPartitionEditorGrid2D::OnMouseButtonDoubleClick(const FGeometry& In
 	{
 		if (InMouseEvent.IsShiftDown())
 		{
-			if (!GetDefault<UWorldPartitionEditorSettings>()->bDisablePIE)
-			{
-				PlayFromHere();
-			}
+			PlayFromHere();
 		}
 		else
 		{
 			MoveCameraHere();
 
-			if (InMouseEvent.IsControlDown() && !GetDefault<UWorldPartitionEditorSettings>()->bDisableLoadingInEditor)
+			if (InMouseEvent.IsControlDown() && WorldPartition->IsLoadingInEditorEnabled())
 			{
 				LoadFromHere();
 			}
diff --git a/Engine/Source/Editor/WorldPartitionEditor/Private/WorldPartitionEditorModule.cpp b/Engine/Source/Editor/WorldPartitionEditor/Private/WorldPartitionEditorModule.cpp
--- a/Engine/Source/Editor/WorldPartitionEditor/Private/WorldPartitionEditorModule.cpp
+++ b/Engine/Source/Editor/WorldPartitionEditor/Private/WorldPartitionEditorModule.cpp
@@ -240,26 +240,6 @@ int32 FWorldPartitionEditorModule::GetMinimapLowQualityWorldUnitsPerPixelThresho
 	return GetDefault<UWorldPartitionEditorSettings>()->MinimapLowQualityWorldUnitsPerPixelThreshold;
 }
 
-bool FWorldPartitionEditorModule::GetDisableLoadingInEditor() const
-{
-	return GetDefault<UWorldPartitionEditorSettings>()->bDisableLoadingInEditor;
-}
-
-void FWorldPartitionEditorModule::SetDisableLoadingInEditor(bool bInDisableLoadingInEditor)
-{
-	GetMutableDefault<UWorldPartitionEditorSettings>()->bDisableLoadingInEditor = bInDisableLoadingInEditor;
-}
-
-bool FWorldPartitionEditorModule::GetDisablePIE() const
-{
-	return GetDefault<UWorldPartitionEditorSettings>()->bDisablePIE;
-}
-
-void FWorldPartitionEditorModule::SetDisablePIE(bool bInDisablePIE)
-{
-	GetMutableDefault<UWorldPartitionEditorSettings>()->bDisablePIE = bInDisablePIE;
-}
-
 void FWorldPartitionEditorModule::OnConvertMap()
 {
 	IContentBrowserSingleton& ContentBrowserSingleton = FModuleManager::LoadModuleChecked<FContentBrowserModule>("ContentBrowser").Get();
@@ -696,8 +676,6 @@ UWorldPartitionEditorSettings::UWorldPartitionEditorSettings()
 	CommandletClass = UWorldPartitionConvertCommandlet::StaticClass();
 	InstancedFoliageGridSize = 25600;
 	MinimapLowQualityWorldUnitsPerPixelThreshold = 12800;
-	bDisableLoadingInEditor = false;
-	bDisablePIE = false;
 }
 
 FString UWorldPartitionConvertOptions::ToCommandletArgs() const
diff --git a/Engine/Source/Editor/WorldPartitionEditor/Public/WorldPartition/WorldPartitionEditorSettings.h b/Engine/Source/Editor/WorldPartitionEditor/Public/WorldPartition/WorldPartitionEditorSettings.h
--- a/Engine/Source/Editor/WorldPartitionEditor/Public/WorldPartition/WorldPartitionEditorSettings.h
+++ b/Engine/Source/Editor/WorldPartitionEditor/Public/WorldPartition/WorldPartitionEditorSettings.h
@@ -15,15 +15,12 @@ class UWorldPartitionEditorSettings : public UDeveloperSettings
 public:
 	UWorldPartitionEditorSettings();
 
-	UPROPERTY(Config, EditAnywhere, Category = MapConversion, Meta = (ToolTip = "Commandlet class to use for World Parition conversion"))
+	UPROPERTY(config, EditAnywhere, Category = MapConversion, meta = (ToolTip = "Commandlet class to use for World Parition conversion"))
 	TSubclassOf<UWorldPartitionConvertCommandlet> CommandletClass;
 
-	UPROPERTY(Config, EditAnywhere, Category = Foliage, Meta = (ClampMin=3200, ToolTip= "Editor grid size used for instance foliage actors in World Partition worlds"))
+	UPROPERTY(config, EditAnywhere, Category = Foliage, meta = (ClampMin=3200, ToolTip= "Editor grid size used for instance foliage actors in World Partition worlds"))
 	int32 InstancedFoliageGridSize;
 
-	UPROPERTY(Config, EditAnywhere, Category = MiniMap, Meta = (ClampMin = 100, ToolTip = "Threshold from which minimap generates a warning if its WorldUnitsPerPixel is above this value"))
+	UPROPERTY(config, EditAnywhere, Category = MiniMap, meta = (ClampMin = 100, ToolTip = "Threshold from which minimap generates a warning if its WorldUnitsPerPixel is above this value"))
 	int32 MinimapLowQualityWorldUnitsPerPixelThreshold;
-
-	bool bDisableLoadingInEditor;
-	bool bDisablePIE;
 };
\ No newline at end of file
diff --git a/Engine/Source/Editor/WorldPartitionEditor/Public/WorldPartitionEditorModule.h b/Engine/Source/Editor/WorldPartitionEditor/Public/WorldPartitionEditorModule.h
--- a/Engine/Source/Editor/WorldPartitionEditor/Public/WorldPartitionEditorModule.h
+++ b/Engine/Source/Editor/WorldPartitionEditor/Public/WorldPartitionEditorModule.h
@@ -47,18 +47,6 @@ public:
 	 */
 	virtual int32 GetMinimapLowQualityWorldUnitsPerPixelThreshold() const override;
 
-	/**
-	 * Returns if loading in the editor is disabled or not.
-	 */
-	virtual bool GetDisableLoadingInEditor() const override;
-	virtual void SetDisableLoadingInEditor(bool bInDisableLoadingInEditor) override;
-
-	/**
-	 * Returns if pie is disabled or not..
-	 */
-	virtual bool GetDisablePIE() const override;
-	virtual void SetDisablePIE(bool bInDisablePIE) override;
-
 	/**
 	 * Convert the specified map to a world partition map.
 	 */
diff --git a/Engine/Source/Runtime/Engine/Private/WorldPartition/WorldPartition.cpp b/Engine/Source/Runtime/Engine/Private/WorldPartition/WorldPartition.cpp
--- a/Engine/Source/Runtime/Engine/Private/WorldPartition/WorldPartition.cpp
+++ b/Engine/Source/Runtime/Engine/Private/WorldPartition/WorldPartition.cpp
@@ -244,6 +244,7 @@ UWorldPartition::UWorldPartition(const FObjectInitializer& ObjectInitializer)
 	, bStreamingWasEnabled(true)
 	, bShouldCheckEnableStreamingWarning(false)
 	, bCanBeUsedByLevelInstance(false)
+	, bEnableLoadingInEditor(true)
 	, bForceGarbageCollection(false)
 	, bForceGarbageCollectionPurge(false)
 	, bEnablingStreamingJustified(false)
@@ -713,17 +714,13 @@ bool UWorldPartition::IsMainWorldPartition() const
 void UWorldPartition::OnPostBugItGoCalled(const FVector& Loc, const FRotator& Rot)
 {
 #if WITH_EDITOR
-	if (GetMutableDefault<UWorldPartitionEditorPerProjectUserSettings>()->GetBugItGoLoadRegion())
+	if (GetMutableDefault<UWorldPartitionEditorPerProjectUserSettings>()->GetBugItGoLoadRegion() && IsLoadingInEditorEnabled())
 	{
-		IWorldPartitionEditorModule& WorldPartitionEditorModule = FModuleManager::LoadModuleChecked<IWorldPartitionEditorModule>("WorldPartitionEditor");
-		if (!WorldPartitionEditorModule.GetDisableLoadingInEditor())
-		{
-			const FVector LoadExtent(GLoadingRangeBugItGo, GLoadingRangeBugItGo, HALF_WORLD_MAX);
-			const FBox LoadCellsBox(Loc - LoadExtent, Loc + LoadExtent);
+		const FVector LoadExtent(GLoadingRangeBugItGo, GLoadingRangeBugItGo, HALF_WORLD_MAX);
+		const FBox LoadCellsBox(Loc - LoadExtent, Loc + LoadExtent);
 
-			UWorldPartitionEditorLoaderAdapter* EditorLoaderAdapter = CreateEditorLoaderAdapter<FLoaderAdapterShape>(World, LoadCellsBox, TEXT("BugItGo"));
-			EditorLoaderAdapter->GetLoaderAdapter()->Load();
-		}
+		UWorldPartitionEditorLoaderAdapter* EditorLoaderAdapter = CreateEditorLoaderAdapter<FLoaderAdapterShape>(World, LoadCellsBox, TEXT("BugItGo"));
+		EditorLoaderAdapter->GetLoaderAdapter()->Load();
 	}
 #endif
 }
diff --git a/Engine/Source/Runtime/Engine/Private/WorldPartition/WorldPartitionConverter.cpp b/Engine/Source/Runtime/Engine/Private/WorldPartition/WorldPartitionConverter.cpp
--- a/Engine/Source/Runtime/Engine/Private/WorldPartition/WorldPartitionConverter.cpp
+++ b/Engine/Source/Runtime/Engine/Private/WorldPartition/WorldPartitionConverter.cpp
@@ -67,6 +67,7 @@ bool FWorldPartitionConverter::Convert()
 		FLevelActorFoldersHelper::SetUseActorFolders(World->PersistentLevel, Parameters.bUseActorFolders);
 		WorldPartition->bEnableStreaming = Parameters.bEnableStreaming;
 		WorldPartition->bStreamingWasEnabled = Parameters.bEnableStreaming;
+		WorldPartition->bEnableLoadingInEditor = Parameters.bEnableStreaming;
 		WorldPartition->SetCanBeUsedByLevelInstance(Parameters.bCanBeUsedByLevelInstance);
 		bCreatedWorldPartition = true;
 	}
diff --git a/Engine/Source/Runtime/Engine/Private/WorldPartition/WorldPartitionEditorHash.cpp b/Engine/Source/Runtime/Engine/Private/WorldPartition/WorldPartitionEditorHash.cpp
--- a/Engine/Source/Runtime/Engine/Private/WorldPartition/WorldPartitionEditorHash.cpp
+++ b/Engine/Source/Runtime/Engine/Private/WorldPartition/WorldPartitionEditorHash.cpp
@@ -1,7 +1,6 @@
 // Copyright Epic Games, Inc. All Rights Reserved.
 
 #include "WorldPartition/WorldPartitionEditorHash.h"
-#include "WorldPartition/IWorldPartitionEditorModule.h"
 
 #if WITH_EDITOR
 #include "Modules/ModuleManager.h"
@@ -11,11 +10,7 @@
 
 UWorldPartitionEditorHash::UWorldPartitionEditorHash(const FObjectInitializer& ObjectInitializer)
 	: Super(ObjectInitializer)
-{
-#if WITH_EDITOR
-	WorldPartitionEditorModule = &FModuleManager::LoadModuleChecked<IWorldPartitionEditorModule>("WorldPartitionEditor");
-#endif
-}
+{}
 
 #if WITH_EDITOR
 UWorldPartitionEditorHash::FForEachIntersectingActorParams::FForEachIntersectingActorParams()
diff --git a/Engine/Source/Runtime/Engine/Private/WorldPartition/WorldPartitionEditorSpatialHash.cpp b/Engine/Source/Runtime/Engine/Private/WorldPartition/WorldPartitionEditorSpatialHash.cpp
--- a/Engine/Source/Runtime/Engine/Private/WorldPartition/WorldPartitionEditorSpatialHash.cpp
+++ b/Engine/Source/Runtime/Engine/Private/WorldPartition/WorldPartitionEditorSpatialHash.cpp
@@ -5,7 +5,7 @@
 #include "WorldPartition/WorldPartitionLog.h"
 #include "WorldPartition/WorldPartitionActorDesc.h"
 #include "WorldPartition/WorldPartitionActorDesc.h"
-#include "WorldPartition/IWorldPartitionEditorModule.h"
+#include "WorldPartition/WorldPartitionEditorPerProjectUserSettings.h"
 
 #include UE_INLINE_GENERATED_CPP_BY_NAME(WorldPartitionEditorSpatialHash)
 
@@ -115,7 +115,7 @@ void UWorldPartitionEditorSpatialHash::HashActor(FWorldPartitionHandle& InActorH
 {
 	check(InActorHandle.IsValid());
 
-	const bool bConsiderActorSpatiallyLoaded = InActorHandle->GetIsSpatiallyLoaded() && !WorldPartitionEditorModule->GetDisableLoadingInEditor();
+	const bool bConsiderActorSpatiallyLoaded = InActorHandle->GetIsSpatiallyLoaded() && GetOuterUWorldPartition()->IsLoadingInEditorEnabled();
 	const FBox ActorBounds = bConsiderActorSpatiallyLoaded ? InActorHandle->GetBounds() : FBox(ForceInit);
 
 #if DO_CHECK
@@ -226,7 +226,7 @@ void UWorldPartitionEditorSpatialHash::UnhashActor(FWorldPartitionHandle& InActo
 {
 	check(InActorHandle.IsValid());
 
-	const bool bConsiderActorSpatiallyLoaded = InActorHandle->GetIsSpatiallyLoaded()  && !WorldPartitionEditorModule->GetDisableLoadingInEditor();
+	const bool bConsiderActorSpatiallyLoaded = InActorHandle->GetIsSpatiallyLoaded() && GetOuterUWorldPartition()->IsLoadingInEditorEnabled();
 	const FBox ActorBounds = bConsiderActorSpatiallyLoaded ? InActorHandle->GetBounds() : FBox(ForceInit);
 
 #if DO_CHECK
diff --git a/Engine/Source/Runtime/Engine/Public/WorldPartition/IWorldPartitionEditorModule.h b/Engine/Source/Runtime/Engine/Public/WorldPartition/IWorldPartitionEditorModule.h
--- a/Engine/Source/Runtime/Engine/Public/WorldPartition/IWorldPartitionEditorModule.h
+++ b/Engine/Source/Runtime/Engine/Public/WorldPartition/IWorldPartitionEditorModule.h
@@ -34,12 +34,6 @@ public:
 	virtual int32 GetPlacementGridSize() const = 0;
 	virtual int32 GetInstancedFoliageGridSize() const = 0;
 	virtual int32 GetMinimapLowQualityWorldUnitsPerPixelThreshold() const = 0;
-	
-	virtual bool GetDisableLoadingInEditor() const = 0;
-	virtual void SetDisableLoadingInEditor(bool bInDisableLoadingInEditor) = 0;
-
-	virtual bool GetDisablePIE() const = 0;
-	virtual void SetDisablePIE(bool bInDisablePIE) = 0;
 
 	/** Triggered when a world is added. */
 	DECLARE_EVENT_OneParam(IWorldPartitionEditorModule, FWorldPartitionCreated, UWorld*);
diff --git a/Engine/Source/Runtime/Engine/Public/WorldPartition/WorldPartition.h b/Engine/Source/Runtime/Engine/Public/WorldPartition/WorldPartition.h
--- a/Engine/Source/Runtime/Engine/Public/WorldPartition/WorldPartition.h
+++ b/Engine/Source/Runtime/Engine/Public/WorldPartition/WorldPartition.h
@@ -202,6 +202,7 @@ public:
 	void OnUserCreatedRegionUnloaded() { check(HasLoadedUserCreatedRegions()); NumUserCreatedLoadedRegions--; }
 
 	bool IsEnablingStreamingJustified() const { return bEnablingStreamingJustified; }
+	bool IsLoadingInEditorEnabled() const { return bEnableLoadingInEditor; }
 #endif
 
 public:
@@ -261,6 +262,10 @@ private:
 	/** Whether Level Instance can reference this partition. */
 	UPROPERTY(EditAnywhere, Category = "WorldPartitionSetup", AdvancedDisplay, meta = (EditConditionHides, HideEditConditionToggle, EditCondition = "!bEnableStreaming"))
 	bool bCanBeUsedByLevelInstance;
+
+	/** When disabled, editor will always load the whole world. */
+	UPROPERTY()
+	bool bEnableLoadingInEditor;
 #endif
 
 public:
diff --git a/Engine/Source/Runtime/Engine/Public/WorldPartition/WorldPartitionEditorHash.h b/Engine/Source/Runtime/Engine/Public/WorldPartition/WorldPartitionEditorHash.h
--- a/Engine/Source/Runtime/Engine/Public/WorldPartition/WorldPartitionEditorHash.h
+++ b/Engine/Source/Runtime/Engine/Public/WorldPartition/WorldPartitionEditorHash.h
@@ -8,8 +8,6 @@
 #include "WorldPartition/WorldPartitionHandle.h"
 #include "WorldPartitionEditorHash.generated.h"
 
-class IWorldPartitionEditorModule;
-
 UCLASS(Abstract, Config=Engine, Within = WorldPartition)
 class ENGINE_API UWorldPartitionEditorHash : public UObject
 {
@@ -46,8 +44,5 @@ public:
 	};
 
 	virtual int32 ForEachIntersectingActor(const FBox& Box, TFunctionRef<void(FWorldPartitionActorDesc*)> InOperation, const FForEachIntersectingActorParams& Params = FForEachIntersectingActorParams()) PURE_VIRTUAL(UWorldPartitionEditorHash::ForEachIntersectingActor, return 0;);
-
-protected:
-	IWorldPartitionEditorModule* WorldPartitionEditorModule;
 #endif
 };
