World Partition - HLOD: New HLOD layers assets are now based on the project default HLOD layer asset

diff --git a/Engine/Source/Editor/WorldPartitionEditor/Private/WorldPartition/HLOD/HLODLayerFactory.cpp b/Engine/Source/Editor/WorldPartitionEditor/Private/WorldPartition/HLOD/HLODLayerFactory.cpp
--- a/Engine/Source/Editor/WorldPartitionEditor/Private/WorldPartition/HLOD/HLODLayerFactory.cpp
+++ b/Engine/Source/Editor/WorldPartitionEditor/Private/WorldPartition/HLOD/HLODLayerFactory.cpp
@@ -6,9 +6,6 @@
 #include "Templates/SubclassOf.h"
 #include "WorldPartition/HLOD/HLODLayer.h"
 
-class FFeedbackContext;
-class UClass;
-class UObject;
 
 UHLODLayerFactory::UHLODLayerFactory(const FObjectInitializer& ObjectInitializer)
 	: Super(ObjectInitializer)
@@ -22,7 +19,20 @@ UHLODLayerFactory::UHLODLayerFactory(const FObjectInitializer& ObjectInitializer
 
 UObject* UHLODLayerFactory::FactoryCreateNew(UClass* Class, UObject* InParent, FName Name, EObjectFlags Flags, UObject* Context, FFeedbackContext* Warn)
 {
-	UHLODLayer* HLODLayer = NewObject<UHLODLayer>(InParent, Class, Name, Flags);
+	UHLODLayer* HLODLayer = nullptr;
+
+	if (UHLODLayer* HLODLayerTemplate = UHLODLayer::GetEngineDefaultHLODLayersSetup())
+	{
+		HLODLayer = DuplicateObject<UHLODLayer>(HLODLayerTemplate, InParent, Name);
+		HLODLayer->ClearFlags(RF_AllFlags);
+		HLODLayer->SetFlags(Flags);
+	}
+	
+	if (!HLODLayer)
+	{
+		HLODLayer = NewObject<UHLODLayer>(InParent, Class, Name, Flags);
+	}
+	
 	check(HLODLayer);
 	return HLODLayer;
 }
