World Partition - HLOD: Made the HLODLayerFactory.h header public

diff --git a/Engine/Source/Editor/WorldPartitionEditor/Private/WorldPartition/HLOD/HLODLayerFactory.h b/Engine/Source/Editor/WorldPartitionEditor/Private/WorldPartition/HLOD/HLODLayerFactory.h
deleted file mode 100644
diff --git a/Engine/Source/Editor/WorldPartitionEditor/Public/WorldPartition/HLOD/HLODLayerFactory.h b/Engine/Source/Editor/WorldPartitionEditor/Public/WorldPartition/HLOD/HLODLayerFactory.h
new file mode 100644
--- /dev/null
+++ b/Engine/Source/Editor/WorldPartitionEditor/Public/WorldPartition/HLOD/HLODLayerFactory.h
@@ -0,0 +1,23 @@
+// Copyright Epic Games, Inc. All Rights Reserved.
+
+#pragma once
+
+#include "Factories/Factory.h"
+#include "UObject/NameTypes.h"
+#include "UObject/UObjectGlobals.h"
+
+#include "HLODLayerFactory.generated.h"
+
+class FFeedbackContext;
+class UClass;
+class UObject;
+
+UCLASS(MinimalAPI, hidecategories = Object)
+class WORLDPARTITIONEDITOR_API UHLODLayerFactory : public UFactory
+{
+	GENERATED_UCLASS_BODY()
+
+	//~ Begin UFactory Interface
+	virtual UObject* FactoryCreateNew(UClass* Class, UObject* InParent, FName Name, EObjectFlags Flags, UObject* Context, FFeedbackContext* Warn) override;
+	//~ Begin UFactory Interface	
+};
