Build fix
#rnx

diff --git a/Engine/Source/Editor/WorldPartitionEditor/Public/WorldPartition/HLOD/HLODLayerFactory.h b/Engine/Source/Editor/WorldPartitionEditor/Public/WorldPartition/HLOD/HLODLayerFactory.h
--- a/Engine/Source/Editor/WorldPartitionEditor/Public/WorldPartition/HLOD/HLODLayerFactory.h
+++ b/Engine/Source/Editor/WorldPartitionEditor/Public/WorldPartition/HLOD/HLODLayerFactory.h
@@ -12,7 +12,7 @@ class FFeedbackContext;
 class UClass;
 class UObject;
 
-UCLASS(MinimalAPI, hidecategories = Object)
+UCLASS()
 class WORLDPARTITIONEDITOR_API UHLODLayerFactory : public UFactory
 {
 	GENERATED_UCLASS_BODY()
