World Partition:
- World Partition Editor is now available when streaming is disabled, but with limitations: no selection, no loading regions, etc.
- Focus on world and minimap generation will now work on non-spatially loaded actors bounds when streaming is disabled.

#rnx

diff --git a/Engine/Source/Editor/WorldPartitionEditor/Private/WorldPartition/SWorldPartitionEditor.cpp b/Engine/Source/Editor/WorldPartitionEditor/Private/WorldPartition/SWorldPartitionEditor.cpp
--- a/Engine/Source/Editor/WorldPartitionEditor/Private/WorldPartition/SWorldPartitionEditor.cpp
+++ b/Engine/Source/Editor/WorldPartitionEditor/Private/WorldPartition/SWorldPartitionEditor.cpp
@@ -71,12 +71,7 @@ TSharedRef<SWidget> SWorldPartitionEditor::ConstructContentWidget()
 		if (UWorldPartition* WorldPartition = World->GetWorldPartition())
 		{
 			check(WorldPartition->World == World);
-
-			if (WorldPartition->IsStreamingEnabled())
-			{
-				EditorName = WorldPartition->GetWorldPartitionEditorName();
-			}
-
+			EditorName = WorldPartition->GetWorldPartitionEditorName();
 			WorldPartition->WorldPartitionEditor = this;
 		}
 	}
diff --git a/Engine/Source/Editor/WorldPartitionEditor/Private/WorldPartition/SWorldPartitionEditorGrid.cpp b/Engine/Source/Editor/WorldPartitionEditor/Private/WorldPartition/SWorldPartitionEditorGrid.cpp
--- a/Engine/Source/Editor/WorldPartitionEditor/Private/WorldPartition/SWorldPartitionEditorGrid.cpp
+++ b/Engine/Source/Editor/WorldPartitionEditor/Private/WorldPartition/SWorldPartitionEditorGrid.cpp
@@ -20,10 +20,7 @@ void SWorldPartitionEditorGrid::Construct(const FArguments& InArgs)
 
 	if (IsDisabled())
 	{
-		const FText Message = !World ? LOCTEXT("WorldPartitionNoWorldForEditorGrid", "No World") :
-									  World->GetWorldPartition() ? 
-									 	LOCTEXT("WorldPartitionHasStreamingDisabled", "World Partition streaming is not enabled for this map") : 
-										LOCTEXT("WorldPartitionMustBeEnabled", "World Partition is not enabled for this map");
+		const FText Message = LOCTEXT("WorldPartitionIsDisabled", "World Partition is disabled for this map");
 
 		ChildSlot
 		.VAlign(VAlign_Center)
@@ -111,7 +108,7 @@ void SWorldPartitionEditorGrid::Refresh()
 
 bool SWorldPartitionEditorGrid::IsDisabled() const
 {
-	return !WorldPartition || !WorldPartition->IsStreamingEnabled();
+	return !WorldPartition;
 }
 
 #undef LOCTEXT_NAMESPACE
diff --git a/Engine/Source/Editor/WorldPartitionEditor/Private/WorldPartition/SWorldPartitionEditorGrid2D.cpp b/Engine/Source/Editor/WorldPartitionEditor/Private/WorldPartition/SWorldPartitionEditorGrid2D.cpp
--- a/Engine/Source/Editor/WorldPartitionEditor/Private/WorldPartition/SWorldPartitionEditorGrid2D.cpp
+++ b/Engine/Source/Editor/WorldPartitionEditor/Private/WorldPartition/SWorldPartitionEditorGrid2D.cpp
@@ -360,6 +360,7 @@ void SWorldPartitionEditorGrid2D::Construct(const FArguments& InArgs)
 						SNew(SCheckBox)
 						.IsChecked(GetMutableDefault<UWorldPartitionEditorPerProjectUserSettings>()->GetShowCellCoords() ? ECheckBoxState::Checked : ECheckBoxState::Unchecked)
 						.IsEnabled(true)
+						.Visibility_Lambda([this]() { return WorldPartition->IsStreamingEnabled() ? EVisibility::Visible : EVisibility::Hidden; })
 						.OnCheckStateChanged(FOnCheckStateChanged::CreateLambda([=](ECheckBoxState State) { GetMutableDefault<UWorldPartitionEditorPerProjectUserSettings>()->SetShowCellCoords(State == ECheckBoxState::Checked); }))
 					]
 					+SHorizontalBox::Slot()
@@ -369,6 +370,7 @@ void SWorldPartitionEditorGrid2D::Construct(const FArguments& InArgs)
 						SNew(STextBlock)
 						.AutoWrapText(true)
 						.IsEnabled(true)
+						.Visibility_Lambda([this]() { return WorldPartition->IsStreamingEnabled() ? EVisibility::Visible : EVisibility::Hidden; })
 						.Text(LOCTEXT("ShowCellCoords", "Show Cell Coords"))
 					]
 					+ SHorizontalBox::Slot()
@@ -1000,7 +1002,7 @@ void SWorldPartitionEditorGrid2D::Tick(const FGeometry& AllottedGeometry, const
 	{
 		if (bShowActors)
 		{
-			if (ActorDesc->GetIsSpatiallyLoaded())
+			if (!WorldPartition->IsStreamingEnabled() || ActorDesc->GetIsSpatiallyLoaded())
 			{
 				ShownActorGuids.Add(ActorDesc->GetGuid());
 			}
@@ -1705,55 +1707,58 @@ int32 SWorldPartitionEditorGrid2D::PaintMinimap(const FGeometry& AllottedGeometr
 			FLinearColor::White
 		);
 
-		F2DRectBooleanSubtract ShadowAreas(MinimapBounds);
-		for (const FLoaderInterface& LoaderInterface : ShownLoaderInterfaces)
+		if (WorldPartition->IsStreamingEnabled())
 		{
-			if (LoaderInterface.IsValid())
+			F2DRectBooleanSubtract ShadowAreas(MinimapBounds);
+			for (const FLoaderInterface& LoaderInterface : ShownLoaderInterfaces)
 			{
-				const IWorldPartitionActorLoaderInterface::ILoaderAdapter* LoaderAdapter = LoaderInterface->GetLoaderAdapter();
-
-				if (LoaderAdapter->GetBoundingBox().IsSet())
+				if (LoaderInterface.IsValid())
 				{
-					const FBox AdapterBounds = *LoaderAdapter->GetBoundingBox();
+					const IWorldPartitionActorLoaderInterface::ILoaderAdapter* LoaderAdapter = LoaderInterface->GetLoaderAdapter();
+
+					if (LoaderAdapter->GetBoundingBox().IsSet())
+					{
+						const FBox AdapterBounds = *LoaderAdapter->GetBoundingBox();
 
-					FVector Origin, Extent;
-					AdapterBounds.GetCenterAndExtents(Origin, Extent);
+						FVector Origin, Extent;
+						AdapterBounds.GetCenterAndExtents(Origin, Extent);
 
-					const FVector2D TopLeftW = FVector2D(Origin - Extent);
-					const FVector2D BottomRightW = FVector2D(Origin + Extent);
-					const FVector2D TopRightW = FVector2D(BottomRightW.X, TopLeftW.Y);
-					const FVector2D BottomLeftW = FVector2D(TopLeftW.X, BottomRightW.Y);
+						const FVector2D TopLeftW = FVector2D(Origin - Extent);
+						const FVector2D BottomRightW = FVector2D(Origin + Extent);
+						const FVector2D TopRightW = FVector2D(BottomRightW.X, TopLeftW.Y);
+						const FVector2D BottomLeftW = FVector2D(TopLeftW.X, BottomRightW.Y);
 
-					const FVector2D TopLeft = WorldToScreen.TransformPoint(TopLeftW);
-					const FVector2D BottomRight = WorldToScreen.TransformPoint(BottomRightW);
-					const FVector2D TopRight = WorldToScreen.TransformPoint(TopRightW);
-					const FVector2D BottomLeft = WorldToScreen.TransformPoint(BottomLeftW);
+						const FVector2D TopLeft = WorldToScreen.TransformPoint(TopLeftW);
+						const FVector2D BottomRight = WorldToScreen.TransformPoint(BottomRightW);
+						const FVector2D TopRight = WorldToScreen.TransformPoint(TopRightW);
+						const FVector2D BottomLeft = WorldToScreen.TransformPoint(BottomLeftW);
 
-					const FBox2D ActorViewBox(TopLeft, BottomRight);
+						const FBox2D ActorViewBox(TopLeft, BottomRight);
 
-					ShadowAreas.SubRect(ActorViewBox);
+						ShadowAreas.SubRect(ActorViewBox);
+					}
 				}
 			}
-		}
 
-		for (const FBox2D& ShadowArea : ShadowAreas.GetRects())
-		{
-			const FPaintGeometry ShadowAreaGeometry = AllottedGeometry.ToPaintGeometry(
-				ShadowArea.Min,
-				ShadowArea.Max - ShadowArea.Min
-			);
+			for (const FBox2D& ShadowArea : ShadowAreas.GetRects())
+			{
+				const FPaintGeometry ShadowAreaGeometry = AllottedGeometry.ToPaintGeometry(
+					ShadowArea.Min,
+					ShadowArea.Max - ShadowArea.Min
+				);
 
-			const FSlateColorBrush ShadowdBrush(FLinearColor::Gray);
-			const FLinearColor ShadowColor(0, 0, 0, 0.75f);
+				const FSlateColorBrush ShadowdBrush(FLinearColor::Gray);
+				const FLinearColor ShadowColor(0, 0, 0, 0.75f);
 
-			FSlateDrawElement::MakeBox(
-				OutDrawElements,
-				++LayerId,
-				ShadowAreaGeometry,
-				&ShadowdBrush,
-				ESlateDrawEffect::None,
-				ShadowColor
-			);
+				FSlateDrawElement::MakeBox(
+					OutDrawElements,
+					++LayerId,
+					ShadowAreaGeometry,
+					&ShadowdBrush,
+					ESlateDrawEffect::None,
+					ShadowColor
+				);
+			}
 		}
 
 		// Minimap outline
@@ -1909,7 +1914,7 @@ FReply SWorldPartitionEditorGrid2D::FocusSelection()
 
 		if (!SelectionBox.IsValid)
 		{
-			SelectionBox = WorldPartition->GetEditorWorldBounds();
+			SelectionBox = WorldPartition->GetRuntimeWorldBounds();
 		}
 	}
 
@@ -2000,22 +2005,25 @@ void SWorldPartitionEditorGrid2D::UpdateTransform() const
 
 void SWorldPartitionEditorGrid2D::UpdateSelectionBox(bool bSnap)
 {
-	const FBox2D SelectBox2D(FVector2D::Min(SelectionStart, SelectionEnd), FVector2D::Max(SelectionStart, SelectionEnd));
-
-	if (SelectBox2D.GetArea() > 0.0f)
+	if (WorldPartition->IsStreamingEnabled())
 	{
-		const float MinX = SelectBox2D.Min.X;
-		const float MinY = SelectBox2D.Min.Y;
-		const float MaxX = SelectBox2D.Max.X;
-		const float MaxY = SelectBox2D.Max.Y;
-		SelectBox = FBox(FVector(MinX, MinY, -HALF_WORLD_MAX), FVector(MaxX, MaxY, HALF_WORLD_MAX));
-
-		const int32 SelectionSnap = bSnap ? GetSelectionSnap() : 1;
-		const float SnapMinX = FMath::GridSnap(MinX - SelectionSnap / 2, SelectionSnap);
-		const float SnapMinY = FMath::GridSnap(MinY - SelectionSnap / 2, SelectionSnap);
-		const float SnapMaxX = FMath::GridSnap(MaxX + SelectionSnap / 2, SelectionSnap);
-		const float SnapMaxY = FMath::GridSnap(MaxY + SelectionSnap / 2, SelectionSnap);
-		SelectBoxGridSnapped = FBox(FVector(SnapMinX, SnapMinY, -HALF_WORLD_MAX), FVector(SnapMaxX, SnapMaxY, HALF_WORLD_MAX));
+		const FBox2D SelectBox2D(FVector2D::Min(SelectionStart, SelectionEnd), FVector2D::Max(SelectionStart, SelectionEnd));
+
+		if (SelectBox2D.GetArea() > 0.0f)
+		{
+			const float MinX = SelectBox2D.Min.X;
+			const float MinY = SelectBox2D.Min.Y;
+			const float MaxX = SelectBox2D.Max.X;
+			const float MaxY = SelectBox2D.Max.Y;
+			SelectBox = FBox(FVector(MinX, MinY, -HALF_WORLD_MAX), FVector(MaxX, MaxY, HALF_WORLD_MAX));
+
+			const int32 SelectionSnap = bSnap ? GetSelectionSnap() : 1;
+			const float SnapMinX = FMath::GridSnap(MinX - SelectionSnap / 2, SelectionSnap);
+			const float SnapMinY = FMath::GridSnap(MinY - SelectionSnap / 2, SelectionSnap);
+			const float SnapMaxX = FMath::GridSnap(MaxX + SelectionSnap / 2, SelectionSnap);
+			const float SnapMaxY = FMath::GridSnap(MaxY + SelectionSnap / 2, SelectionSnap);
+			SelectBoxGridSnapped = FBox(FVector(SnapMinX, SnapMinY, -HALF_WORLD_MAX), FVector(SnapMaxX, SnapMaxY, HALF_WORLD_MAX));
+		}
 	}
 }
 
diff --git a/Engine/Source/Runtime/Engine/Private/WorldPartition/WorldPartition.cpp b/Engine/Source/Runtime/Engine/Private/WorldPartition/WorldPartition.cpp
--- a/Engine/Source/Runtime/Engine/Private/WorldPartition/WorldPartition.cpp
+++ b/Engine/Source/Runtime/Engine/Private/WorldPartition/WorldPartition.cpp
@@ -1626,13 +1626,13 @@ void UWorldPartition::RemapSoftObjectPath(FSoftObjectPath& ObjectPath)
 FBox UWorldPartition::GetEditorWorldBounds() const
 {
 	check(EditorHash);
-	return EditorHash->GetEditorWorldBounds();
+	return IsStreamingEnabled() ? EditorHash->GetEditorWorldBounds() : EditorHash->GetNonSpatialBounds();
 }
 
 FBox UWorldPartition::GetRuntimeWorldBounds() const
 {
 	check(EditorHash);
-	return EditorHash->GetRuntimeWorldBounds();
+	return IsStreamingEnabled() ? EditorHash->GetRuntimeWorldBounds() : EditorHash->GetNonSpatialBounds();
 }
 #endif
 
diff --git a/Engine/Source/Runtime/Engine/Private/WorldPartition/WorldPartitionEditorSpatialHash.cpp b/Engine/Source/Runtime/Engine/Private/WorldPartition/WorldPartitionEditorSpatialHash.cpp
--- a/Engine/Source/Runtime/Engine/Private/WorldPartition/WorldPartitionEditorSpatialHash.cpp
+++ b/Engine/Source/Runtime/Engine/Private/WorldPartition/WorldPartitionEditorSpatialHash.cpp
@@ -17,8 +17,11 @@ TAutoConsoleVariable<bool> CVarEnableSpatialHashValidation(TEXT("wp.Editor.Enabl
 
 UWorldPartitionEditorSpatialHash::UWorldPartitionEditorSpatialHash(const FObjectInitializer& ObjectInitializer)
 	: Super(ObjectInitializer)
-#if WITH_EDITOR
+#if WITH_EDITORONLY_DATA
 	, CellSize(12800)
+	, EditorBounds(ForceInit)
+	, RuntimeBounds(ForceInit)
+	, NonSpatialBounds(ForceInit)
 	, bBoundsDirty(false)
 #endif
 {}
@@ -50,25 +53,36 @@ FBox UWorldPartitionEditorSpatialHash::GetRuntimeWorldBounds() const
 	return RuntimeBounds;
 }
 
+FBox UWorldPartitionEditorSpatialHash::GetNonSpatialBounds() const
+{
+	return NonSpatialBounds;
+}
+
 void UWorldPartitionEditorSpatialHash::Tick(float DeltaSeconds)
 {
 	if (bBoundsDirty)
 	{
 		FBox NewEditorBounds(ForceInit);
-		FBox NewRuntimeBounds(ForceInit);
 
+		RuntimeBounds.Init();
 		for (FCell* Cell: Cells)
 		{
 			NewEditorBounds += Cell->Bounds;
 
 			for (const FWorldPartitionHandle& Actor : Cell->Actors)
 			{
-				NewRuntimeBounds += Actor->GetBounds();
+				RuntimeBounds += Actor->GetBounds();
 			}
 		}
 
+		NonSpatialBounds.Init();
+		for (const FWorldPartitionHandle& Actor : AlwaysLoadedCell->Actors)
+		{
+			NonSpatialBounds += Actor->GetBounds();
+		}
+
 		const int32 OldLevel = GetLevelForBox(EditorBounds);
-		check(OldLevel == HashLevels.Num() - 1);
+		check(!EditorBounds.IsValid || OldLevel == HashLevels.Num() - 1);
 
 		const int32 NewLevel = GetLevelForBox(NewEditorBounds);
 		check(NewLevel <= OldLevel);		
@@ -79,8 +93,7 @@ void UWorldPartitionEditorSpatialHash::Tick(float DeltaSeconds)
 		}
 
 		EditorBounds = NewEditorBounds;
-		RuntimeBounds = NewRuntimeBounds;
-
+		
 		bBoundsDirty = false;
 	}
 
@@ -126,6 +139,7 @@ void UWorldPartitionEditorSpatialHash::HashActor(FWorldPartitionHandle& InActorH
 	if (!bConsiderActorSpatiallyLoaded)
 	{
 		AlwaysLoadedCell->Actors.Add(InActorHandle);
+		NonSpatialBounds += InActorHandle->GetBounds();
 	}
 	else
 	{
@@ -278,11 +292,11 @@ void UWorldPartitionEditorSpatialHash::UnhashActor(FWorldPartitionHandle& InActo
 						CurrentCellCoord = ParentCellCoord;
 					}
 				}
-
-				bBoundsDirty = true;
 			}
 		});
 	}
+
+	bBoundsDirty = true;
 }
 
 int32 UWorldPartitionEditorSpatialHash::ForEachIntersectingActor(const FBox& Box, TFunctionRef<void(FWorldPartitionActorDesc*)> InOperation, const FForEachIntersectingActorParams& Params)
diff --git a/Engine/Source/Runtime/Engine/Public/WorldPartition/WorldPartitionEditorHash.h b/Engine/Source/Runtime/Engine/Public/WorldPartition/WorldPartitionEditorHash.h
--- a/Engine/Source/Runtime/Engine/Public/WorldPartition/WorldPartitionEditorHash.h
+++ b/Engine/Source/Runtime/Engine/Public/WorldPartition/WorldPartitionEditorHash.h
@@ -22,6 +22,7 @@ public:
 	virtual FName GetWorldPartitionEditorName() const PURE_VIRTUAL(UWorldPartitionEditorHash::GetWorldPartitionEditorName, return FName(NAME_None););
 	virtual FBox GetEditorWorldBounds() const  PURE_VIRTUAL(UWorldPartitionEditorHash::GetEditorWorldBounds, return FBox(ForceInit););
 	virtual FBox GetRuntimeWorldBounds() const  PURE_VIRTUAL(UWorldPartitionEditorHash::GetRuntimeWorldBounds, return FBox(ForceInit););
+	virtual FBox GetNonSpatialBounds() const PURE_VIRTUAL(UWorldPartitionEditorHash::GetNonSpatialBounds, return FBox(ForceInit););
 	virtual void Tick(float DeltaSeconds) PURE_VIRTUAL(UWorldPartitionEditorHash::Tick, return;);
 
 	virtual void HashActor(FWorldPartitionHandle& InActorHandle) PURE_VIRTUAL(UWorldPartitionEditorHash::HashActor, ;);
diff --git a/Engine/Source/Runtime/Engine/Public/WorldPartition/WorldPartitionEditorSpatialHash.h b/Engine/Source/Runtime/Engine/Public/WorldPartition/WorldPartitionEditorSpatialHash.h
--- a/Engine/Source/Runtime/Engine/Public/WorldPartition/WorldPartitionEditorSpatialHash.h
+++ b/Engine/Source/Runtime/Engine/Public/WorldPartition/WorldPartitionEditorSpatialHash.h
@@ -192,6 +192,7 @@ public:
 	virtual FName GetWorldPartitionEditorName() const override;
 	virtual FBox GetEditorWorldBounds() const override;
 	virtual FBox GetRuntimeWorldBounds() const override;
+	virtual FBox GetNonSpatialBounds() const override;
 	virtual void Tick(float DeltaSeconds) override;
 
 	virtual void HashActor(FWorldPartitionHandle& InActorHandle) override;
@@ -219,6 +220,7 @@ private:
 	
 	FBox EditorBounds;
 	FBox RuntimeBounds;
+	FBox NonSpatialBounds;
 	bool bBoundsDirty;
 
 #if DO_CHECK
