Remove ContentBundleHasErrorColumn & remove DataLayerHasErrorColumn header icon
DataLayerHasErrorColumn to be removed at a later stage when the problematic UI operations potentially causing errors can be addressed

diff --git a/Engine/Source/Editor/DataLayerEditor/Private/DataLayer/DataLayerOutlinerHasErrorColumn.cpp b/Engine/Source/Editor/DataLayerEditor/Private/DataLayer/DataLayerOutlinerHasErrorColumn.cpp
--- a/Engine/Source/Editor/DataLayerEditor/Private/DataLayer/DataLayerOutlinerHasErrorColumn.cpp
+++ b/Engine/Source/Editor/DataLayerEditor/Private/DataLayer/DataLayerOutlinerHasErrorColumn.cpp
@@ -44,11 +44,10 @@ SHeaderRow::FColumn::FArguments FDataLayerOutlinerHasErrorsColumn::ConstructHead
 		.VAlignHeader(VAlign_Center)
 		.HAlignCell(HAlign_Center)
 		.VAlignCell(VAlign_Center)
-		.DefaultTooltip(FText::FromName(GetColumnID()))
 		[
-			SNew(SImage)
-			.Image(FAppStyle::GetBrush(TEXT("Icons.Error")))
-			.ColorAndOpacity(FSlateColor::UseForeground())
+			SNew(STextBlock)
+			.Margin(FMargin(0.0f))
+			.Text(FText::FromString(TEXT("")))
 		];
 }
 
diff --git a/Engine/Source/Editor/WorldPartitionEditor/Private/WorldPartition/ContentBundle/Outliner/ContentBundleHasErrorColumn.cpp b/Engine/Source/Editor/WorldPartitionEditor/Private/WorldPartition/ContentBundle/Outliner/ContentBundleHasErrorColumn.cpp
deleted file mode 100644
diff --git a/Engine/Source/Editor/WorldPartitionEditor/Private/WorldPartition/ContentBundle/Outliner/ContentBundleStatusColumn.cpp b/Engine/Source/Editor/WorldPartitionEditor/Private/WorldPartition/ContentBundle/Outliner/ContentBundleStatusColumn.cpp
--- a/Engine/Source/Editor/WorldPartitionEditor/Private/WorldPartition/ContentBundle/Outliner/ContentBundleStatusColumn.cpp
+++ b/Engine/Source/Editor/WorldPartitionEditor/Private/WorldPartition/ContentBundle/Outliner/ContentBundleStatusColumn.cpp
@@ -58,6 +58,11 @@ const TSharedRef<SWidget> FContentBundleOutlinerStatusColumn::ConstructRowWidget
 							TSharedPtr<FContentBundleEditor> ContentBundleEditorPin = ContentBundleTreeItem->GetContentBundleEditorPin();
 							if (ContentBundleEditorPin != nullptr)
 							{
+								if (ContentBundleEditorPin->GetStatus() == EContentBundleStatus::FailedToInject)
+								{
+									return FText::Format(LOCTEXT("ConsultLogForErrors", "{0} (Consult Log For Errors)"), UEnum::GetDisplayValueAsText(ContentBundleEditorPin->GetStatus()));
+								}
+								
 								return UEnum::GetDisplayValueAsText(ContentBundleEditorPin->GetStatus());
 							}
 
diff --git a/Engine/Source/Editor/WorldPartitionEditor/Private/WorldPartition/ContentBundle/SContentBundleBrowser.cpp b/Engine/Source/Editor/WorldPartitionEditor/Private/WorldPartition/ContentBundle/SContentBundleBrowser.cpp
--- a/Engine/Source/Editor/WorldPartitionEditor/Private/WorldPartition/ContentBundle/SContentBundleBrowser.cpp
+++ b/Engine/Source/Editor/WorldPartitionEditor/Private/WorldPartition/ContentBundle/SContentBundleBrowser.cpp
@@ -8,7 +8,6 @@
 #include "WorldPartition/ContentBundle/Outliner/ContentBundleEditingColumn.h"
 #include "WorldPartition/ContentBundle/Outliner/ContentBundleStatusColumn.h"
 #include "WorldPartition/ContentBundle/Outliner/ContentBundleClientColumn.h"
-#include "WorldPartition/ContentBundle/Outliner/ContentBundleHasErrorColumn.h"
 #include "WorldPartition/ContentBundle/Outliner/ContentBundleActorCountColumn.h"
 
 namespace ContentBundleOutlinerPrivate
@@ -19,8 +18,7 @@ namespace ContentBundleOutlinerPrivate
 		BuiltIn,
 		Status,
 		Editing,
-		ActorCount,
-		HasError
+		ActorCount
 	};
 }
 
@@ -38,7 +36,6 @@ void SContentBundleBrowser::Construct(const FArguments& InArgs)
 	InitOptions.ColumnMap.Add(FContentBundleOutlinerStatusColumn::GetID(), FSceneOutlinerColumnInfo(ESceneOutlinerColumnVisibility::Visible, ContentBundleOutlinerPrivate::EColumnPriority::Status, FCreateSceneOutlinerColumn::CreateLambda([](ISceneOutliner& InSceneOutliner) { return MakeShareable(new FContentBundleOutlinerStatusColumn(InSceneOutliner)); }), false));
 	InitOptions.ColumnMap.Add(FContentBundleOutlinerEditingColumn::GetID(), FSceneOutlinerColumnInfo(ESceneOutlinerColumnVisibility::Visible, ContentBundleOutlinerPrivate::EColumnPriority::Editing, FCreateSceneOutlinerColumn::CreateLambda([](ISceneOutliner& InSceneOutliner) { return MakeShareable(new FContentBundleOutlinerEditingColumn(InSceneOutliner)); }), false));
 	InitOptions.ColumnMap.Add(FContentBundleOutlinerActorCountColumn::GetID(), FSceneOutlinerColumnInfo(ESceneOutlinerColumnVisibility::Visible, ContentBundleOutlinerPrivate::EColumnPriority::ActorCount, FCreateSceneOutlinerColumn::CreateLambda([](ISceneOutliner& InSceneOutliner) { return MakeShareable(new FContentBundleOutlinerActorCountColumn(InSceneOutliner)); }), false));
-	InitOptions.ColumnMap.Add(FContentBundleOutlinerHasErrorColumn::GetID(), FSceneOutlinerColumnInfo(ESceneOutlinerColumnVisibility::Visible, ContentBundleOutlinerPrivate::EColumnPriority::HasError, FCreateSceneOutlinerColumn::CreateLambda([](ISceneOutliner& InSceneOutliner) { return MakeShareable(new FContentBundleOutlinerHasErrorColumn(InSceneOutliner)); }), false));
 	
 	ContentBundleOutliner = SNew(SContentBundleOutliner, InitOptions)
 		.IsEnabled(FSlateApplication::Get().GetNormalExecutionAttribute());
diff --git a/Engine/Source/Editor/WorldPartitionEditor/Public/WorldPartition/ContentBundle/Outliner/ContentBundleHasErrorColumn.h b/Engine/Source/Editor/WorldPartitionEditor/Public/WorldPartition/ContentBundle/Outliner/ContentBundleHasErrorColumn.h
deleted file mode 100644
