World Partition - HLOD: Added a "Build HLOD" button in the World Partition HLOD details panel

diff --git a/Engine/Source/Editor/WorldPartitionEditor/Private/WorldPartition/Customizations/WorldPartitionHLODDetailsCustomization.cpp b/Engine/Source/Editor/WorldPartitionEditor/Private/WorldPartition/Customizations/WorldPartitionHLODDetailsCustomization.cpp
new file mode 100644
--- /dev/null
+++ b/Engine/Source/Editor/WorldPartitionEditor/Private/WorldPartition/Customizations/WorldPartitionHLODDetailsCustomization.cpp
@@ -0,0 +1,97 @@
+// Copyright Epic Games, Inc. All Rights Reserved.
+
+#include "WorldPartitionHLODDetailsCustomization.h"
+#include "Widgets/DeclarativeSyntaxSupport.h"
+#include "Widgets/Input/SButton.h"
+#include "Editor.h"
+#include "Engine/Selection.h"
+#include "PropertyHandle.h"
+#include "DetailLayoutBuilder.h"
+#include "DetailCategoryBuilder.h"
+#include "DetailWidgetRow.h"
+#include "WorldPartition/WorldPartition.h"
+#include "WorldPartition/HLOD/HLODActor.h"
+
+#include "Algo/Transform.h"
+#include "Algo/ForEach.h"
+
+#include "ComponentRecreateRenderStateContext.h"
+
+#define LOCTEXT_NAMESPACE "FWorldPartitionHLODDetailsCustomization"
+
+TSharedRef<IDetailCustomization> FWorldPartitionHLODDetailsCustomization::MakeInstance()
+{
+	return MakeShareable(new FWorldPartitionHLODDetailsCustomization);
+}
+
+void FWorldPartitionHLODDetailsCustomization::CustomizeDetails(IDetailLayoutBuilder& DetailLayoutBuilder)
+{
+	SelectedObjects = DetailLayoutBuilder.GetSelectedObjects();
+
+	IDetailCategoryBuilder& PluginCategory = DetailLayoutBuilder.EditCategory("Tools");
+	PluginCategory.AddCustomRow(LOCTEXT("HLODTools", "HLOD Tools"), false)
+	[
+		SNew(SVerticalBox)
+		+SVerticalBox::Slot()
+		[
+			SNew(SButton)
+			.Text(LOCTEXT("BuildHLODButtonText", "Build HLOD"))
+			.OnClicked(this, &FWorldPartitionHLODDetailsCustomization::OnBuildHLOD)
+		]
+	];
+}
+
+TArray<AWorldPartitionHLOD*> FWorldPartitionHLODDetailsCustomization::GetSelectedHLODActors() const
+{
+	TArray<AWorldPartitionHLOD*> SelectedHLODActors;
+	SelectedHLODActors.Reserve(SelectedObjects.Num());
+
+	auto IsA_WPHLOD = [](TWeakObjectPtr<UObject> InObject) { return InObject.Get() && InObject->IsA<AWorldPartitionHLOD>(); };
+	auto Cast_WPHLOD = [](TWeakObjectPtr<UObject> InObject) { return Cast<AWorldPartitionHLOD>(InObject); };
+	Algo::TransformIf(SelectedObjects, SelectedHLODActors, IsA_WPHLOD, Cast_WPHLOD);
+
+	return SelectedHLODActors;
+}
+
+FReply FWorldPartitionHLODDetailsCustomization::OnBuildHLOD()
+{
+	TArray<AWorldPartitionHLOD*> SelectedHLODActors = GetSelectedHLODActors();
+
+	// Gather all components
+	TArray<UActorComponent*> ActorComponents;
+	Algo::ForEach(SelectedHLODActors, [&ActorComponents](AWorldPartitionHLOD* HLODActor) { ActorComponents.Append(HLODActor->GetComponents().Array()); });
+
+	// Recreate render states for all the components we're about to process
+	FGlobalComponentRecreateRenderStateContext RecreateRenderStateContext(ActorComponents);
+
+	// Use a material update context to signal any change made to materials to the render thread... Exclude the recreate render state flag as this step is already performed
+	// by the FGlobalComponentRecreateRenderStateContext above.
+	FMaterialUpdateContext MaterialUpdateContext(FMaterialUpdateContext::EOptions::Default & ~FMaterialUpdateContext::EOptions::RecreateRenderStates);
+	for (UActorComponent* ActorComponent : ActorComponents)
+	{
+		if (UStaticMeshComponent* SMComponent = Cast<UStaticMeshComponent>(ActorComponent))
+		{
+			TArray<UMaterialInterface*> UsedMaterials;
+			SMComponent->GetUsedMaterials(UsedMaterials);
+
+			// Add all used materials to the material update context.
+			for (UMaterialInterface* UsedMaterial : UsedMaterials)
+			{
+				MaterialUpdateContext.AddMaterialInterface(UsedMaterial);
+			}
+		}
+	}
+
+	// Build HLODs
+	for (AWorldPartitionHLOD* HLODActor : SelectedHLODActors)
+	{
+		HLODActor->BuildHLOD(true);
+	};
+	
+	// Force refresh the UI so that any change to the selected actors are properly shown to the user
+	GEditor->NoteSelectionChange();
+
+	return FReply::Handled();
+}
+
+#undef LOCTEXT_NAMESPACE
diff --git a/Engine/Source/Editor/WorldPartitionEditor/Private/WorldPartition/Customizations/WorldPartitionHLODDetailsCustomization.h b/Engine/Source/Editor/WorldPartitionEditor/Private/WorldPartition/Customizations/WorldPartitionHLODDetailsCustomization.h
new file mode 100644
--- /dev/null
+++ b/Engine/Source/Editor/WorldPartitionEditor/Private/WorldPartition/Customizations/WorldPartitionHLODDetailsCustomization.h
@@ -0,0 +1,32 @@
+// Copyright Epic Games, Inc. All Rights Reserved.
+
+#pragma once
+
+#include "CoreMinimal.h"
+#include "Input/Reply.h"
+#include "IDetailCustomization.h"
+
+class IDetailLayoutBuilder;
+class AWorldPartitionHLOD;
+
+class FWorldPartitionHLODDetailsCustomization : public IDetailCustomization
+{
+public:
+	/** Makes a new instance of this detail layout class for a specific detail view requesting it */
+	static TSharedRef<IDetailCustomization> MakeInstance();
+
+protected:
+	FWorldPartitionHLODDetailsCustomization()
+	{}
+
+	FReply OnBuildHLOD();
+
+private:
+	/** IDetailCustomization interface */
+	virtual void CustomizeDetails(IDetailLayoutBuilder& InDetailBuilder) override;
+
+	TArray<AWorldPartitionHLOD*> GetSelectedHLODActors() const;
+
+private:
+	TArray<TWeakObjectPtr<UObject>> SelectedObjects;
+};
\ No newline at end of file
diff --git a/Engine/Source/Editor/WorldPartitionEditor/Private/WorldPartitionEditorModule.cpp b/Engine/Source/Editor/WorldPartitionEditor/Private/WorldPartitionEditorModule.cpp
--- a/Engine/Source/Editor/WorldPartitionEditor/Private/WorldPartitionEditorModule.cpp
+++ b/Engine/Source/Editor/WorldPartitionEditor/Private/WorldPartitionEditorModule.cpp
@@ -10,6 +10,7 @@
 #include "WorldPartition/SWorldPartitionEditor.h"
 #include "WorldPartition/SWorldPartitionEditorGridSpatialHash.h"
 #include "WorldPartition/Customizations/WorldPartitionDetailsCustomization.h"
+#include "WorldPartition/Customizations/WorldPartitionHLODDetailsCustomization.h"
 #include "WorldPartition/Customizations/WorldPartitionRuntimeSpatialHashDetailsCustomization.h"
 #include "WorldPartition/SWorldPartitionConvertDialog.h"
 #include "WorldPartition/WorldPartitionConvertOptions.h"
@@ -141,6 +142,7 @@ void FWorldPartitionEditorModule::StartupModule()
 	FPropertyEditorModule& PropertyEditor = FModuleManager::LoadModuleChecked<FPropertyEditorModule>("PropertyEditor");
 	PropertyEditor.RegisterCustomClassLayout("WorldPartition", FOnGetDetailCustomizationInstance::CreateStatic(&FWorldPartitionDetails::MakeInstance));
 	PropertyEditor.RegisterCustomClassLayout("WorldPartitionRuntimeSpatialHash", FOnGetDetailCustomizationInstance::CreateStatic(&FWorldPartitionRuntimeSpatialHashDetails::MakeInstance));
+	PropertyEditor.RegisterCustomClassLayout("WorldPartitionHLOD", FOnGetDetailCustomizationInstance::CreateStatic(&FWorldPartitionHLODDetailsCustomization::MakeInstance));
 }
 
 void FWorldPartitionEditorModule::ShutdownModule()
