World Partition Editor: grid is now rendered using an adaptive grid cell sizes.

#rnx

diff --git a/Engine/Source/Editor/WorldPartitionEditor/Private/WorldPartition/SWorldPartitionEditorGrid2D.cpp b/Engine/Source/Editor/WorldPartitionEditor/Private/WorldPartition/SWorldPartitionEditorGrid2D.cpp
--- a/Engine/Source/Editor/WorldPartitionEditor/Private/WorldPartition/SWorldPartitionEditorGrid2D.cpp
+++ b/Engine/Source/Editor/WorldPartitionEditor/Private/WorldPartition/SWorldPartitionEditorGrid2D.cpp
@@ -615,7 +615,7 @@ void SWorldPartitionEditorGrid2D::PlayFromHere()
 
 void SWorldPartitionEditorGrid2D::LoadFromHere()
 {
-	const float SelectionSnap = GetSelectionSnap();
+	const double SelectionSnap = GetSelectionSnap();
 	const FVector LoadLocation = FVector(MouseCursorPosWorld, 0);
 	const FVector LoadExtent(SelectionSnap, SelectionSnap, HALF_WORLD_MAX);
 	FBox LoadCellsBox(LoadLocation - LoadExtent, LoadLocation + LoadExtent);
@@ -650,7 +650,7 @@ bool SWorldPartitionEditorGrid2D::IsInteractive() const
 	return !IsFollowPlayerInPIE();
 }
 
-int32 SWorldPartitionEditorGrid2D::GetSelectionSnap() const
+int64 SWorldPartitionEditorGrid2D::GetSelectionSnap() const
 {
 	return 1;
 }
@@ -908,7 +908,7 @@ FReply SWorldPartitionEditorGrid2D::OnMouseWheel(const FGeometry& MyGeometry, co
 	{
 		FVector2D MousePosLocalSpace = MouseCursorPos - MyGeometry.GetLocalSize() * 0.5f;
 		FVector2D P0 = MousePosLocalSpace / Scale;
-		float Delta = 1.0f + FMath::Abs(MouseEvent.GetWheelDelta() / 4.0f);
+		float Delta = 1.0f + FMath::Abs(MouseEvent.GetWheelDelta() / 8.0f);
 		Scale = FMath::Clamp(Scale * (MouseEvent.GetWheelDelta() > 0 ? Delta : (1.0f / Delta)), 0.00000001f, 10.0f);
 		FVector2D P1 = MousePosLocalSpace / Scale;
 		Trans += (P1 - P0);
@@ -2021,7 +2021,7 @@ void SWorldPartitionEditorGrid2D::UpdateSelectionBox(bool bSnap)
 			const float MaxY = SelectBox2D.Max.Y;
 			SelectBox = FBox(FVector(MinX, MinY, -HALF_WORLD_MAX), FVector(MaxX, MaxY, HALF_WORLD_MAX));
 
-			const int32 SelectionSnap = bSnap ? GetSelectionSnap() : 1;
+			const int64 SelectionSnap = bSnap ? GetSelectionSnap() : 1;
 			const float SnapMinX = FMath::GridSnap(MinX - SelectionSnap / 2, SelectionSnap);
 			const float SnapMinY = FMath::GridSnap(MinY - SelectionSnap / 2, SelectionSnap);
 			const float SnapMaxX = FMath::GridSnap(MaxX + SelectionSnap / 2, SelectionSnap);
diff --git a/Engine/Source/Editor/WorldPartitionEditor/Private/WorldPartition/SWorldPartitionEditorGridSpatialHash.cpp b/Engine/Source/Editor/WorldPartitionEditor/Private/WorldPartition/SWorldPartitionEditorGridSpatialHash.cpp
--- a/Engine/Source/Editor/WorldPartitionEditor/Private/WorldPartition/SWorldPartitionEditorGridSpatialHash.cpp
+++ b/Engine/Source/Editor/WorldPartitionEditor/Private/WorldPartition/SWorldPartitionEditorGridSpatialHash.cpp
@@ -19,19 +19,28 @@ void SWorldPartitionEditorGridSpatialHash::Construct(const FArguments& InArgs)
 	SWorldPartitionEditorGrid2D::Construct(SWorldPartitionEditorGrid::FArguments().InWorld(InArgs._InWorld));
 }
 
-int32 SWorldPartitionEditorGridSpatialHash::GetSelectionSnap() const
+int64 SWorldPartitionEditorGridSpatialHash::GetSelectionSnap() const
+{
+	float FadeRatio;
+	float CellScreenSize;
+	return GetSelectionSnap(FadeRatio, CellScreenSize);
+}
+
+int64 SWorldPartitionEditorGridSpatialHash::GetSelectionSnap(float& FadeRatio, float& CellScreenSize) const
 {
 	UWorldPartitionEditorSpatialHash* EditorSpatialHash = (UWorldPartitionEditorSpatialHash*)WorldPartition->EditorHash;
 	
-	int32 EffectiveCellSize = EditorSpatialHash->CellSize;
-	for (const UWorldPartitionEditorSpatialHash::FCellNodeHashLevel& HashLevel : EditorSpatialHash->HashLevels)
+	int64 EffectiveCellSize = EditorSpatialHash->CellSize;
+	for (;;)
 	{
-		const FVector2D CellScreenSize = WorldToScreen.TransformVector(FVector2D(EffectiveCellSize, EffectiveCellSize));
-		if (CellScreenSize.X > 64)
+		CellScreenSize = WorldToScreen.TransformVector(FVector2D(EffectiveCellSize, EffectiveCellSize)).X;
+
+		if (CellScreenSize > WantedCellScreenSize)
 		{
+			FadeRatio = (CellScreenSize - WantedCellScreenSize) / WantedCellScreenSize;
 			break;
 		}
-
+	
 		EffectiveCellSize *= 2;
 	}
 
@@ -41,35 +50,51 @@ int32 SWorldPartitionEditorGridSpatialHash::GetSelectionSnap() const
 int32 SWorldPartitionEditorGridSpatialHash::PaintGrid(const FGeometry& AllottedGeometry, const FSlateRect& MyCullingRect, FSlateWindowElementList& OutDrawElements, int32 LayerId) const
 {
 	UWorldPartitionEditorSpatialHash* EditorSpatialHash = (UWorldPartitionEditorSpatialHash*)WorldPartition->EditorHash;
-	
+
 	// Found the best cell size depending on the current zoom
-	const int64 EffectiveCellSize = GetSelectionSnap();
+	float FadeRatio;
+	float CellScreenSize;
+	int64 EffectiveCellSize = GetSelectionSnap(FadeRatio, CellScreenSize);
 	
 	// Compute visible rect
 	const FBox2D ViewRect(FVector2D(ForceInitToZero), AllottedGeometry.GetLocalSize());
 	const FBox2D ViewRectWorld(ScreenToWorld.TransformPoint(ViewRect.Min), ScreenToWorld.TransformPoint(ViewRect.Max));
 
-	FBox VisibleGridRectWorld(
+	const FBox VisibleGridRectWorld(
 		FVector(
-			FMath::Max(FMath::FloorToFloat(EditorSpatialHash->EditorBounds.Min.X / EffectiveCellSize) * EffectiveCellSize, ViewRectWorld.Min.X),
-			FMath::Max(FMath::FloorToFloat(EditorSpatialHash->EditorBounds.Min.Y / EffectiveCellSize) * EffectiveCellSize, ViewRectWorld.Min.Y),
+			FMath::Max(FMath::FloorToFloat(ViewRectWorld.Min.X / EffectiveCellSize) * EffectiveCellSize, ViewRectWorld.Min.X),
+			FMath::Max(FMath::FloorToFloat(ViewRectWorld.Min.Y / EffectiveCellSize) * EffectiveCellSize, ViewRectWorld.Min.Y),
 			FMath::FloorToFloat(EditorSpatialHash->EditorBounds.Min.Z / EffectiveCellSize) * EffectiveCellSize
 		),
 		FVector(
-			FMath::Min(FMath::CeilToFloat(EditorSpatialHash->EditorBounds.Max.X / EffectiveCellSize) * EffectiveCellSize, ViewRectWorld.Max.X),
-			FMath::Min(FMath::CeilToFloat(EditorSpatialHash->EditorBounds.Max.Y / EffectiveCellSize) * EffectiveCellSize, ViewRectWorld.Max.Y),
+			FMath::Min(FMath::CeilToFloat(ViewRectWorld.Max.X / EffectiveCellSize) * EffectiveCellSize, ViewRectWorld.Max.X),
+			FMath::Min(FMath::CeilToFloat(ViewRectWorld.Max.Y / EffectiveCellSize) * EffectiveCellSize, ViewRectWorld.Max.Y),
 			FMath::CeilToFloat(EditorSpatialHash->EditorBounds.Max.Z / EffectiveCellSize) * EffectiveCellSize
 		)
 	);
 
-	// Shadow whole grid area
+	// Paint minimap
 	{
+		const FBox VisibleMinimapRectWorld(
+			FVector(
+				FMath::Max(FMath::FloorToFloat(EditorSpatialHash->RuntimeBounds.Min.X / EditorSpatialHash->CellSize) * EditorSpatialHash->CellSize, ViewRectWorld.Min.X),
+				FMath::Max(FMath::FloorToFloat(EditorSpatialHash->RuntimeBounds.Min.Y / EditorSpatialHash->CellSize) * EditorSpatialHash->CellSize, ViewRectWorld.Min.Y),
+				FMath::FloorToFloat(EditorSpatialHash->RuntimeBounds.Min.Z / EditorSpatialHash->CellSize) * EditorSpatialHash->CellSize
+			),
+			FVector(
+				FMath::Min(FMath::CeilToFloat(EditorSpatialHash->RuntimeBounds.Max.X / EditorSpatialHash->CellSize) * EditorSpatialHash->CellSize, ViewRectWorld.Max.X),
+				FMath::Min(FMath::CeilToFloat(EditorSpatialHash->RuntimeBounds.Max.Y / EditorSpatialHash->CellSize) * EditorSpatialHash->CellSize, ViewRectWorld.Max.Y),
+				FMath::CeilToFloat(EditorSpatialHash->RuntimeBounds.Max.Z / EditorSpatialHash->CellSize) * EditorSpatialHash->CellSize
+			)
+		);
+
+		// Shadow minimap area
 		FSlateColorBrush ShadowBrush(FLinearColor::Black);
 		FLinearColor ShadowColor(FLinearColor(0.0f, 0.0f, 0.0f, 0.5f));
 
 		FPaintGeometry GridGeometry = AllottedGeometry.ToPaintGeometry(
-			WorldToScreen.TransformPoint(FVector2D(VisibleGridRectWorld.Min)),
-			WorldToScreen.TransformPoint(FVector2D(VisibleGridRectWorld.Max)) - WorldToScreen.TransformPoint(FVector2D(VisibleGridRectWorld.Min))
+			WorldToScreen.TransformPoint(FVector2D(VisibleMinimapRectWorld.Min)),
+			WorldToScreen.TransformPoint(FVector2D(VisibleMinimapRectWorld.Max)) - WorldToScreen.TransformPoint(FVector2D(VisibleMinimapRectWorld.Min))
 		);
 
 		FSlateDrawElement::MakeBox(
@@ -80,58 +105,100 @@ int32 SWorldPartitionEditorGridSpatialHash::PaintGrid(const FGeometry& AllottedG
 			ESlateDrawEffect::None,
 			ShadowColor
 		);
-	}
 
-	// Paint minimap
-	LayerId = PaintMinimap(AllottedGeometry, MyCullingRect, OutDrawElements, LayerId);
+		LayerId = PaintMinimap(AllottedGeometry, MyCullingRect, OutDrawElements, LayerId);
+	}
 
 	// Paint grid lines
 	if (ToBox2D(VisibleGridRectWorld).GetArea() > 0.0f)
 	{
-		const FLinearColor Color = FLinearColor(0.1f, 0.1f, 0.1f, 1.f);
-
-		UE::Math::TIntVector2<int64> TopLeftW(
-			FMath::FloorToFloat(VisibleGridRectWorld.Min.X / EffectiveCellSize) * EffectiveCellSize,
-			FMath::FloorToFloat(VisibleGridRectWorld.Min.Y / EffectiveCellSize) * EffectiveCellSize
-		);
-
-		UE::Math::TIntVector2<int64> BottomRightW(
-			FMath::CeilToFloat(VisibleGridRectWorld.Max.X / EffectiveCellSize) * EffectiveCellSize,
-			FMath::CeilToFloat(VisibleGridRectWorld.Max.Y / EffectiveCellSize) * EffectiveCellSize
-		);
+		auto PaintGridLines = [this, &OutDrawElements, &LayerId, &AllottedGeometry](const FBox& VisibleGridRectWorld, float EffectiveCellSize, float Thickness, const FLinearColor& Color)
+		{
+			if (Color.A > 0.0f)
+			{
+				UE::Math::TIntVector2<int64> TopLeftW(
+					FMath::FloorToFloat(VisibleGridRectWorld.Min.X / EffectiveCellSize) * EffectiveCellSize,
+					FMath::FloorToFloat(VisibleGridRectWorld.Min.Y / EffectiveCellSize) * EffectiveCellSize
+				);
 
-		TArray<FVector2D> LinePoints;
-		LinePoints.SetNum(2);
+				UE::Math::TIntVector2<int64> BottomRightW(
+					FMath::CeilToFloat(VisibleGridRectWorld.Max.X / EffectiveCellSize) * EffectiveCellSize,
+					FMath::CeilToFloat(VisibleGridRectWorld.Max.Y / EffectiveCellSize) * EffectiveCellSize
+				);
 
-		// Horizontal
-		for (int64 i=TopLeftW.Y; i<=BottomRightW.Y; i+=EffectiveCellSize)
-		{
-			FVector2D LineStartH(TopLeftW.X, i);
-			FVector2D LineEndH(BottomRightW.X, i);
+				TArray<FVector2D> LinePoints;
+				LinePoints.SetNum(2);
 
-			LinePoints[0] = WorldToScreen.TransformPoint(LineStartH);
-			LinePoints[1] = WorldToScreen.TransformPoint(LineEndH);
+				// Horizontal
+				for (int64 i=TopLeftW.Y; i<=BottomRightW.Y; i+=EffectiveCellSize)
+				{
+					FVector2D LineStartH(TopLeftW.X, i);
+					FVector2D LineEndH(BottomRightW.X, i);
+
+					LinePoints[0] = WorldToScreen.TransformPoint(LineStartH);
+					LinePoints[1] = WorldToScreen.TransformPoint(LineEndH);
+
+					FSlateDrawElement::MakeLines(
+						OutDrawElements, 
+						LayerId, 
+						AllottedGeometry.ToPaintGeometry(), 
+						LinePoints, 
+						ESlateDrawEffect::None, 
+						Color, 
+						false, 
+						Thickness
+					);
+				}
 
-			FSlateDrawElement::MakeLines(OutDrawElements, LayerId, AllottedGeometry.ToPaintGeometry(), LinePoints, ESlateDrawEffect::None, Color, false, 1.0f);
-		}
+				// Vertical
+				for (int64 i=TopLeftW.X; i<=BottomRightW.X; i+=EffectiveCellSize)
+				{
+					FVector2D LineStartH(i, TopLeftW.Y);
+					FVector2D LineEndH(i, BottomRightW.Y);
+
+					LinePoints[0] = WorldToScreen.TransformPoint(LineStartH);
+					LinePoints[1] = WorldToScreen.TransformPoint(LineEndH);
+
+					FSlateDrawElement::MakeLines(
+						OutDrawElements, LayerId, 
+						AllottedGeometry.ToPaintGeometry(), 
+						LinePoints, 
+						ESlateDrawEffect::None, 
+						Color, 
+						false, 
+						Thickness
+					);
+				}
+			}
+		};
 
-		// Vertical
-		for (int64 i=TopLeftW.X; i<=BottomRightW.X; i+=EffectiveCellSize)
+		auto PaintGridLinesFaded = [this, &EffectiveCellSize, EditorSpatialHash, &PaintGridLines, &VisibleGridRectWorld, CellScreenSize](int32 Granularity, float Thickness)
 		{
-			FVector2D LineStartH(i, TopLeftW.Y);
-			FVector2D LineEndH(i, BottomRightW.Y);
-
-			LinePoints[0] = WorldToScreen.TransformPoint(LineStartH);
-			LinePoints[1] = WorldToScreen.TransformPoint(LineEndH);
+			if (EffectiveCellSize >= EditorSpatialHash->CellSize * Granularity)
+			{
+				const float Alpha = FMath::Min(CellScreenSize / (WantedCellScreenSize * Granularity), 1.0f);
+				PaintGridLines(VisibleGridRectWorld, EffectiveCellSize / Granularity, Thickness, FLinearColor(0.1f, 0.1f, 0.1f, Alpha * 0.25f));
+			}
+		};
 
-			FSlateDrawElement::MakeLines(OutDrawElements, LayerId, AllottedGeometry.ToPaintGeometry(), LinePoints, ESlateDrawEffect::None, Color, false, 1.0f);
-		}
+		PaintGridLinesFaded(4, 1.0f);
+		PaintGridLinesFaded(2, 1.0f);
+		PaintGridLinesFaded(1, 2.0f);
 
 		// Draw coordinates
 		if ((EffectiveCellSize == EditorSpatialHash->CellSize) && GetMutableDefault<UWorldPartitionEditorPerProjectUserSettings>()->GetShowCellCoords())
 		{
+			UE::Math::TIntVector2<int64> TopLeftW(
+				FMath::FloorToFloat(VisibleGridRectWorld.Min.X / EffectiveCellSize) * EffectiveCellSize,
+				FMath::FloorToFloat(VisibleGridRectWorld.Min.Y / EffectiveCellSize) * EffectiveCellSize
+			);
+
+			UE::Math::TIntVector2<int64> BottomRightW(
+				FMath::CeilToFloat(VisibleGridRectWorld.Max.X / EffectiveCellSize) * EffectiveCellSize,
+				FMath::CeilToFloat(VisibleGridRectWorld.Max.Y / EffectiveCellSize) * EffectiveCellSize
+			);
+
 			const TSharedRef<FSlateFontMeasure> FontMeasure = FSlateApplication::Get().GetRenderer()->GetFontMeasureService();
-			const FVector2D CellScreenSize = WorldToScreen.TransformVector(FVector2D(EffectiveCellSize, EffectiveCellSize));
 
 			FSlateFontInfo CoordsFont;
 			FVector2D DefaultCoordTextSize;
@@ -143,17 +210,17 @@ int32 SWorldPartitionEditorGridSpatialHash::PaintGrid(const FGeometry& AllottedG
 				CoordsFont = FCoreStyle::GetDefaultFontStyle("Bold", DesiredFontSize);
 				DefaultCoordTextSize = FontMeasure->Measure(DefaultCoordText, CoordsFont);
 
-				if (CellScreenSize.X > DefaultCoordTextSize.X)
+				if (CellScreenSize > DefaultCoordTextSize.X)
 				{
 					bNeedsGradient = (DesiredFontSize == 8);
 					break;
 				}
 			}
 
-			if (CellScreenSize.X > DefaultCoordTextSize.X)
+			if (CellScreenSize > DefaultCoordTextSize.X)
 			{
-				static float GradientDistance = 64.0f;
-				float ColorGradient = bNeedsGradient ? FMath::Min((CellScreenSize.X - DefaultCoordTextSize.X) / GradientDistance, 1.0f) : 1.0f;
+				static float GradientDistance = WantedCellScreenSize;
+				float ColorGradient = bNeedsGradient ? FMath::Min((CellScreenSize - DefaultCoordTextSize.X) / GradientDistance, 1.0f) : 1.0f;
 				const FLinearColor CoordTextColor(1.0f, 1.0f, 1.0f, ColorGradient);
 
 				for (int64 y = TopLeftW.Y; y < BottomRightW.Y; y += EffectiveCellSize)
diff --git a/Engine/Source/Editor/WorldPartitionEditor/Public/WorldPartition/SWorldPartitionEditorGrid2D.h b/Engine/Source/Editor/WorldPartitionEditor/Public/WorldPartition/SWorldPartitionEditorGrid2D.h
--- a/Engine/Source/Editor/WorldPartitionEditor/Public/WorldPartition/SWorldPartitionEditorGrid2D.h
+++ b/Engine/Source/Editor/WorldPartitionEditor/Public/WorldPartition/SWorldPartitionEditorGrid2D.h
@@ -52,7 +52,7 @@ protected:
 	bool IsFollowPlayerInPIE() const;
 	bool IsInteractive() const;
 
-	virtual int32 GetSelectionSnap() const;
+	virtual int64 GetSelectionSnap() const;
 
 	virtual void Tick(const FGeometry& AllottedGeometry, const double InCurrentTime, const float InDeltaTime) override;
 	virtual FReply OnMouseButtonDown(const FGeometry& MyGeometry, const FPointerEvent& MouseEvent) override;
diff --git a/Engine/Source/Editor/WorldPartitionEditor/Public/WorldPartition/SWorldPartitionEditorGridSpatialHash.h b/Engine/Source/Editor/WorldPartitionEditor/Public/WorldPartition/SWorldPartitionEditorGridSpatialHash.h
--- a/Engine/Source/Editor/WorldPartitionEditor/Public/WorldPartition/SWorldPartitionEditorGridSpatialHash.h
+++ b/Engine/Source/Editor/WorldPartitionEditor/Public/WorldPartition/SWorldPartitionEditorGridSpatialHash.h
@@ -11,6 +11,11 @@ public:
 
 	void Construct(const FArguments& InArgs);
 
-	virtual int32 GetSelectionSnap() const override;
+	virtual int64 GetSelectionSnap() const override;
 	virtual int32 PaintGrid(const FGeometry& AllottedGeometry, const FSlateRect& MyCullingRect, FSlateWindowElementList& OutDrawElements, int32 LayerId) const override;
+
+private:
+	int64 GetSelectionSnap(float& FadeRatio, float& CellScreenSize) const;
+
+	inline static float WantedCellScreenSize = 64.0f;
 };
