World Partition - HLOD: Added null check to HLOD actor's "Build HLOD" button processing
#rnx

diff --git a/Engine/Source/Editor/WorldPartitionEditor/Private/WorldPartition/Customizations/WorldPartitionHLODDetailsCustomization.cpp b/Engine/Source/Editor/WorldPartitionEditor/Private/WorldPartition/Customizations/WorldPartitionHLODDetailsCustomization.cpp
--- a/Engine/Source/Editor/WorldPartitionEditor/Private/WorldPartition/Customizations/WorldPartitionHLODDetailsCustomization.cpp
+++ b/Engine/Source/Editor/WorldPartitionEditor/Private/WorldPartition/Customizations/WorldPartitionHLODDetailsCustomization.cpp
@@ -79,7 +79,10 @@ FReply FWorldPartitionHLODDetailsCustomization::OnBuildHLOD()
 			// Add all used materials to the material update context.
 			for (UMaterialInterface* UsedMaterial : UsedMaterials)
 			{
-				MaterialUpdateContext.AddMaterialInterface(UsedMaterial);
+				if (UsedMaterial)
+				{
+					MaterialUpdateContext.AddMaterialInterface(UsedMaterial);
+				}
 			}
 		}
 	}
