World Partition Editor:
- When an actor is selected, show both current and on-disk bounds.
- Show faded subcells grids when zoomed in further than editor cell sizes.

#rnx

diff --git a/Engine/Source/Editor/WorldPartitionEditor/Private/WorldPartition/SWorldPartitionEditorGrid2D.cpp b/Engine/Source/Editor/WorldPartitionEditor/Private/WorldPartition/SWorldPartitionEditorGrid2D.cpp
--- a/Engine/Source/Editor/WorldPartitionEditor/Private/WorldPartition/SWorldPartitionEditorGrid2D.cpp
+++ b/Engine/Source/Editor/WorldPartitionEditor/Private/WorldPartition/SWorldPartitionEditorGrid2D.cpp
@@ -1326,73 +1326,88 @@ uint32 SWorldPartitionEditorGrid2D::PaintActors(const FGeometry& AllottedGeometr
 		bool bIsBoundsEdgeHoveredFound = false;
 		for (const FWorldPartitionActorDescViewBoundsProxy& ActorDescView: ActorDescList)
 		{
-			const FBox ActorBounds = ActorDescView.GetBounds();
-			FVector Origin, Extent;
-			ActorBounds.GetCenterAndExtents(Origin, Extent);
+			auto ShowActorBox = [this, &AllottedGeometry, &OutDrawElements, &LayerId, &bIsBoundsEdgeHoveredFound, &DrawActorLabel](const FBox& ActorBounds, bool bIsSelected, bool bIsSpatiallyLoaded, const FName ActorLabel)
+			{
+				FVector Origin, Extent;
+				ActorBounds.GetCenterAndExtents(Origin, Extent);
 
-			FVector2D TopLeftW = FVector2D(Origin - Extent);
-			FVector2D BottomRightW = FVector2D(Origin + Extent);
-			FVector2D TopRightW = FVector2D(BottomRightW.X, TopLeftW.Y);
-			FVector2D BottomLeftW = FVector2D(TopLeftW.X, BottomRightW.Y);
+				FVector2D TopLeftW = FVector2D(Origin - Extent);
+				FVector2D BottomRightW = FVector2D(Origin + Extent);
+				FVector2D TopRightW = FVector2D(BottomRightW.X, TopLeftW.Y);
+				FVector2D BottomLeftW = FVector2D(TopLeftW.X, BottomRightW.Y);
 
-			FVector2D TopLeft = WorldToScreen.TransformPoint(TopLeftW);
-			FVector2D BottomRight = WorldToScreen.TransformPoint(BottomRightW);
-			FVector2D TopRight = WorldToScreen.TransformPoint(TopRightW);
-			FVector2D BottomLeft = WorldToScreen.TransformPoint(BottomLeftW);
+				FVector2D TopLeft = WorldToScreen.TransformPoint(TopLeftW);
+				FVector2D BottomRight = WorldToScreen.TransformPoint(BottomRightW);
+				FVector2D TopRight = WorldToScreen.TransformPoint(TopRightW);
+				FVector2D BottomLeft = WorldToScreen.TransformPoint(BottomLeftW);
 
-			FBox2D ActorViewBox(TopLeft, BottomRight);
+				FBox2D ActorViewBox(TopLeft, BottomRight);
 
-			const float MinimumAreaCull = 32.0f;
-			const float AreaFadeDistance = 128.0f;
-			if ((Extent.Size2D() < KINDA_SMALL_NUMBER) || (ActorViewBox.GetArea() > MinimumAreaCull))
-			{
-				const UClass* ActorClass = ActorDescView.GetActorNativeClass();
-				const AActor* Actor = ActorDescView.GetActor();
-				const FPaintGeometry ActorGeometry = AllottedGeometry.ToPaintGeometry(TopLeft, BottomRight - TopLeft);
-				const FPaintGeometry ActorGeometryShadow = AllottedGeometry.ToPaintGeometry(TopLeft - FVector2D::One(), BottomRight - TopLeft + FVector2D::One() * 2.0);
-				const bool bIsSelected = Actor ? Actor->IsSelected() : false;
-				const bool bIsBoundsEdgeHovered = !bIsBoundsEdgeHoveredFound && IsBoundsEdgeHovered(MouseCursorPos, ActorViewBox, 10.0f);
-				const float ActorColorGradient = FMath::Cube(FMath::Clamp((ActorViewBox.GetArea() - MinimumAreaCull) / AreaFadeDistance, 0.0f, 1.0f));
-				const float ActorBrightness = ActorDescView.GetIsSpatiallyLoaded() ? 1.0f : 0.3f;
-
-				// Only draw the first edge bounds hovered
-				bIsBoundsEdgeHoveredFound |= bIsBoundsEdgeHovered;
+				const float MinimumAreaCull = 32.0f;
+				const float AreaFadeDistance = 128.0f;
+				if ((Extent.Size2D() < KINDA_SMALL_NUMBER) || (ActorViewBox.GetArea() > MinimumAreaCull))
+				{
+					const FPaintGeometry ActorGeometry = AllottedGeometry.ToPaintGeometry(TopLeft, BottomRight - TopLeft);
+					const FPaintGeometry ActorGeometryShadow = AllottedGeometry.ToPaintGeometry(TopLeft - FVector2D::One(), BottomRight - TopLeft + FVector2D::One() * 2.0);
+					const bool bIsBoundsEdgeHovered = !bIsBoundsEdgeHoveredFound && IsBoundsEdgeHovered(MouseCursorPos, ActorViewBox, 10.0f);
+					const float ActorColorGradient = FMath::Cube(FMath::Clamp((ActorViewBox.GetArea() - MinimumAreaCull) / AreaFadeDistance, 0.0f, 1.0f));
+					const float ActorBrightness = bIsSpatiallyLoaded ? 1.0f : 0.3f;
+
+					// Only draw the first edge bounds hovered
+					bIsBoundsEdgeHoveredFound |= bIsBoundsEdgeHovered;
 				
-				FLinearColor ActorColor(ActorBrightness, ActorBrightness, ActorBrightness, bIsBoundsEdgeHovered ? 1.0f : ActorColorGradient);
+					FLinearColor ActorColor(ActorBrightness, ActorBrightness, ActorBrightness, bIsBoundsEdgeHovered ? 1.0f : ActorColorGradient);
 
-				if (bIsSelected || bIsBoundsEdgeHovered)
-				{
-					const FName ActorLabel = ActorDescView.GetActorLabel();
-					if (!ActorLabel.IsNone())
+					if (bIsSelected || bIsBoundsEdgeHovered)
+					{
+						if (!ActorLabel.IsNone())
+						{
+							DrawActorLabel(ActorLabel.ToString(), bIsBoundsEdgeHovered ? MouseCursorPos : ActorViewBox.GetCenter(), ActorGeometry, FLinearColor::Yellow, SmallLayoutFont, true);
+						}
+
+						ActorColor = FLinearColor::Yellow.CopyWithNewOpacity(ActorColor.A);
+					}
+					else if ((SelectBoxGridSnapped.GetVolume() > 0) && SelectBoxGridSnapped.Intersect(ActorBounds))
 					{
-						DrawActorLabel(ActorLabel.ToString(), bIsBoundsEdgeHovered ? MouseCursorPos : ActorViewBox.GetCenter(), ActorGeometry, FLinearColor::Yellow, SmallLayoutFont, true);
+						ActorColor = FLinearColor::White.CopyWithNewOpacity(ActorColor.A);
 					}
 
-					ActorColor = FLinearColor::Yellow.CopyWithNewOpacity(ActorColor.A);
-				}
-				else if ((SelectBoxGridSnapped.GetVolume() > 0) && SelectBoxGridSnapped.Intersect(ActorDescView.GetBounds()))
-				{
-					ActorColor = FLinearColor::White.CopyWithNewOpacity(ActorColor.A);
+					FSlateDrawElement::MakeBox(
+						OutDrawElements,
+						++LayerId,
+						ActorGeometryShadow,
+						FAppStyle::GetBrush(TEXT("Border")),
+						ESlateDrawEffect::None,
+						FLinearColor::Black.CopyWithNewOpacity(ActorColor.A)
+					);
+
+					FSlateDrawElement::MakeBox(
+						OutDrawElements,
+						++LayerId,
+						ActorGeometry,
+						FAppStyle::GetBrush(TEXT("Border")),
+						ESlateDrawEffect::None,
+						ActorColor
+					);
 				}
+			};
 
-				FSlateDrawElement::MakeBox(
-					OutDrawElements,
-					++LayerId,
-					ActorGeometryShadow,
-					FAppStyle::GetBrush(TEXT("Border")),
-					ESlateDrawEffect::None,
-					FLinearColor::Black.CopyWithNewOpacity(ActorColor.A)
-				);
+			const FBox ActorBounds = ActorDescView.GetBounds();
+			const AActor* Actor = ActorDescView.GetActor();
+			const bool bIsSelected = Actor ? Actor->IsSelected() : false;
+			const bool bIsSpatiallyLoaded = ActorDescView.GetIsSpatiallyLoaded();
+			const FName ActorLabel = ActorDescView.GetActorLabel();
 
-				FSlateDrawElement::MakeBox(
-					OutDrawElements,
-					++LayerId,
-					ActorGeometry,
-					FAppStyle::GetBrush(TEXT("Border")),
-					ESlateDrawEffect::None,
-					ActorColor
-				);
+			if (bIsSelected)
+			{
+				const FBox ActorDescBounds = ActorDescView.GetActorDesc()->GetBounds();
+				if (!ActorDescBounds.Equals(ActorBounds, 1.0f))
+				{
+					ShowActorBox(ActorDescBounds, false, bIsSpatiallyLoaded, ActorLabel);
+				}
 			}
+
+			ShowActorBox(ActorBounds, bIsSelected, bIsSpatiallyLoaded, ActorLabel);
 		};
 	}
 
diff --git a/Engine/Source/Editor/WorldPartitionEditor/Private/WorldPartition/SWorldPartitionEditorGridSpatialHash.cpp b/Engine/Source/Editor/WorldPartitionEditor/Private/WorldPartition/SWorldPartitionEditorGridSpatialHash.cpp
--- a/Engine/Source/Editor/WorldPartitionEditor/Private/WorldPartition/SWorldPartitionEditorGridSpatialHash.cpp
+++ b/Engine/Source/Editor/WorldPartitionEditor/Private/WorldPartition/SWorldPartitionEditorGridSpatialHash.cpp
@@ -174,11 +174,8 @@ int32 SWorldPartitionEditorGridSpatialHash::PaintGrid(const FGeometry& AllottedG
 
 		auto PaintGridLinesFaded = [this, &EffectiveCellSize, EditorSpatialHash, &PaintGridLines, &VisibleGridRectWorld, CellScreenSize](int32 Granularity, float Thickness)
 		{
-			if (EffectiveCellSize >= EditorSpatialHash->CellSize * Granularity)
-			{
-				const float Alpha = FMath::Min(CellScreenSize / (WantedCellScreenSize * Granularity), 1.0f);
-				PaintGridLines(VisibleGridRectWorld, EffectiveCellSize / Granularity, Thickness, FLinearColor(0.1f, 0.1f, 0.1f, Alpha * 0.25f));
-			}
+			const float Alpha = FMath::Min(CellScreenSize / (WantedCellScreenSize * Granularity), (EffectiveCellSize >= EditorSpatialHash->CellSize * Granularity) ? 1.0f : 0.5f);
+			PaintGridLines(VisibleGridRectWorld, EffectiveCellSize / Granularity, Thickness, FLinearColor(0.1f, 0.1f, 0.1f, Alpha * 0.25f));
 		};
 
 		PaintGridLinesFaded(4, 1.0f);
