World Partition: don't automatically invoke the world partition editor tab anymore, let it open and disabled if the world don't support it.

#rnx

diff --git a/Engine/Source/Editor/WorldPartitionEditor/Private/WorldPartitionEditorModule.cpp b/Engine/Source/Editor/WorldPartitionEditor/Private/WorldPartitionEditorModule.cpp
--- a/Engine/Source/Editor/WorldPartitionEditor/Private/WorldPartitionEditorModule.cpp
+++ b/Engine/Source/Editor/WorldPartitionEditor/Private/WorldPartitionEditorModule.cpp
@@ -164,8 +164,6 @@ void FWorldPartitionEditorModule::ShutdownModule()
 
 		}
 
-		FEditorDelegates::MapChange.RemoveAll(this);
-
 		UToolMenus::UnRegisterStartupCallback(this);
 		UToolMenus::UnregisterOwner(this);
 	}
@@ -209,8 +207,6 @@ void FWorldPartitionEditorModule::RegisterMenus()
 		FSlateIcon(FAppStyle::GetAppStyleSetName(), "DeveloperTools.MenuIcon"),
 		FUIAction(FExecuteAction::CreateRaw(this, &FWorldPartitionEditorModule::OnConvertMap))
 	));
-
-	FEditorDelegates::MapChange.AddRaw(this, &FWorldPartitionEditorModule::OnMapChanged);
 }
 
 TSharedRef<SWidget> FWorldPartitionEditorModule::CreateWorldPartitionEditor()
@@ -615,31 +611,6 @@ bool FWorldPartitionEditorModule::BuildLandscapeSplineMeshes(UWorld* InWorld)
 	return true;
 }
 
-void FWorldPartitionEditorModule::OnMapChanged(uint32 MapFlags)
-{
-	if (MapFlags == MapChangeEventFlags::NewMap)
-	{
-		FLevelEditorModule* LevelEditorModule = FModuleManager::Get().GetModulePtr<FLevelEditorModule>("LevelEditor");
-	
-		TSharedPtr<FTabManager> LevelEditorTabManager = LevelEditorModule ? LevelEditorModule->GetLevelEditorTabManager() : nullptr;
-
-		// If the world opened is a world partition world spawn the world partition tab if not open.
-		UWorld* EditorWorld = GEditor->GetEditorWorldContext().World();
-		if (EditorWorld && EditorWorld->IsPartitionedWorld())
-		{
-			if(LevelEditorTabManager && !WorldPartitionTab.IsValid())
-			{
-				WorldPartitionTab = LevelEditorTabManager->TryInvokeTab(WorldPartitionEditorTabId);
-			}
-		}
-		else if(TSharedPtr<SDockTab> WorldPartitionTabPin = WorldPartitionTab.Pin())
-		{
-			// close the WP tab if not a world partition world
-			WorldPartitionTabPin->RequestCloseTab();
-		}
-	}
-}
-
 TSharedRef<SDockTab> FWorldPartitionEditorModule::SpawnWorldPartitionTab(const FSpawnTabArgs& Args)
 {
 	TSharedRef<SDockTab> NewTab =
diff --git a/Engine/Source/Editor/WorldPartitionEditor/Public/WorldPartitionEditorModule.h b/Engine/Source/Editor/WorldPartitionEditor/Public/WorldPartitionEditorModule.h
--- a/Engine/Source/Editor/WorldPartitionEditor/Public/WorldPartitionEditorModule.h
+++ b/Engine/Source/Editor/WorldPartitionEditor/Public/WorldPartitionEditorModule.h
@@ -79,10 +79,8 @@ public:
 	 * Creates a Content Bundle Browser widget
 	 */
 	TSharedRef<class SWidget> CreateContentBundleBrowser();
-private:
-	/** Called when the level editors map changes. We will determine if the new map is a valid world partition world and close world partition tabs if not */
-	void OnMapChanged(uint32 MapFlags);
 
+private:
 	/** Register menus */
 	void RegisterMenus();
 
