WorldPartition Editor: Add BugIt menu (prints BugItGo command in console and to clipboard)

diff --git a/Engine/Source/Editor/WorldPartitionEditor/Private/WorldPartition/SWorldPartitionEditorGrid2D.cpp b/Engine/Source/Editor/WorldPartitionEditor/Private/WorldPartition/SWorldPartitionEditorGrid2D.cpp
--- a/Engine/Source/Editor/WorldPartitionEditor/Private/WorldPartition/SWorldPartitionEditorGrid2D.cpp
+++ b/Engine/Source/Editor/WorldPartitionEditor/Private/WorldPartition/SWorldPartitionEditorGrid2D.cpp
@@ -4,6 +4,7 @@
 #include "ActorFactories/ActorFactory.h"
 #include "Brushes/SlateColorBrush.h"
 #include "Builders/CubeBuilder.h"
+#include "Editor.h"
 #include "Editor/EditorEngine.h"
 #include "Engine/Selection.h"
 #include "Engine/Texture2D.h"
@@ -24,6 +25,7 @@
 #include "WorldBrowserModule.h"
 #include "WorldPartition/LoaderAdapter/LoaderAdapterShape.h"
 #include "WorldPartition/WorldPartition.h"
+#include "WorldPartition/WorldPartitionLog.h"
 #include "WorldPartition/WorldPartitionActorDesc.h"
 #include "WorldPartition/WorldPartitionActorDescView.h"
 #include "WorldPartition/WorldPartitionEditorHash.h"
@@ -34,6 +36,7 @@
 #include "WorldPartition/WorldPartitionMiniMap.h"
 #include "WorldPartition/WorldPartitionMiniMapHelper.h"
 #include "WorldPartition/WorldPartitionMiniMapVolume.h"
+#include "HAL/PlatformApplicationMisc.h"
 
 #include "IAssetViewport.h"
 #include "Kismet2/DebuggerCommands.h"
@@ -153,6 +156,14 @@ static bool IsBoundsEdgeHovered(const FVector2D& Point, const FBox2D& Bounds, fl
 	return DistanceToPoint > UE_KINDA_SMALL_NUMBER && DistanceToPoint < Size;
 }
 
+static bool HitTestZFromLocation(UWorld* World, const FVector2D& WorldLocation, FHitResult& OutResult)
+{
+	const FVector TraceStart(WorldLocation.X, WorldLocation.Y, HALF_WORLD_MAX);
+	const FVector TraceEnd(WorldLocation.X, WorldLocation.Y, -HALF_WORLD_MAX);
+	const FCollisionQueryParams TraceParams(SCENE_QUERY_STAT(SWorldPartitionEditorGrid2D_HitTestZFromLocation), true);
+	return World->LineTraceSingleByChannel(OutResult, TraceStart, TraceEnd, ECC_WorldStatic, TraceParams);
+}
+
 template <class T>
 void ForEachIntersectingLoaderAdapters(UWorldPartition* WorldPartition, const FBox& SelectBox, T Func)
 {
@@ -238,7 +249,8 @@ void SWorldPartitionEditorGrid2D::FEditorCommands::RegisterCommands()
 	UI_COMMAND(LoadSelectedRegions, "Load Selected Regions", "Load the selected regions.", EUserInterfaceActionType::Button, FInputChord());
 	UI_COMMAND(UnloadSelectedRegions, "Unload Selected Regions", "Unload the selected regions.", EUserInterfaceActionType::Button, FInputChord());
 	UI_COMMAND(ConvertSelectedRegionsToActors, "Convert Selected Regions To Actors", "Convert the selected regions to actors.", EUserInterfaceActionType::Button, FInputChord());
-	UI_COMMAND(MoveCameraHere, "Move Camera Here", "Move the camera to the selected position.", EUserInterfaceActionType::Button, FInputChord());
+	UI_COMMAND(MoveCameraHere, "Move Camera Here", "Move the camera to the selected location.", EUserInterfaceActionType::Button, FInputChord());
+	UI_COMMAND(BugItHere, "Bug It", "Log BugItGo command of the selected location to console (also adds it to clipboard).", EUserInterfaceActionType::Button, FInputChord());
 	UI_COMMAND(PlayFromHere, "Play From Here", "Play from here.", EUserInterfaceActionType::Button, FInputChord());
 	UI_COMMAND(LoadFromHere, "Load From Here", "Load from here.", EUserInterfaceActionType::Button, FInputChord());
 }
@@ -455,6 +467,7 @@ void SWorldPartitionEditorGrid2D::Construct(const FArguments& InArgs)
 	ActionList.MapAction(Commands.UnloadSelectedRegions, FExecuteAction::CreateSP(this, &SWorldPartitionEditorGrid2D::UnloadSelectedRegions), FCanExecuteAction::CreateLambda(CanUnloadSelectedRegions));
 	ActionList.MapAction(Commands.ConvertSelectedRegionsToActors, FExecuteAction::CreateSP(this, &SWorldPartitionEditorGrid2D::ConvertSelectedRegionsToActors), FCanExecuteAction::CreateLambda(CanConvertSelectedRegionsToActors));
 	ActionList.MapAction(Commands.MoveCameraHere, FExecuteAction::CreateSP(this, &SWorldPartitionEditorGrid2D::MoveCameraHere));
+	ActionList.MapAction(Commands.BugItHere, FExecuteAction::CreateSP(this, &SWorldPartitionEditorGrid2D::BugItHere));
 	ActionList.MapAction(Commands.PlayFromHere, FExecuteAction::CreateSP(this, &SWorldPartitionEditorGrid2D::PlayFromHere));
 	ActionList.MapAction(Commands.LoadFromHere, FExecuteAction::CreateSP(this, &SWorldPartitionEditorGrid2D::LoadFromHere));
 }
@@ -569,15 +582,30 @@ void SWorldPartitionEditorGrid2D::ConvertSelectedRegionsToActors()
 	Refresh();
 }
 
+void SWorldPartitionEditorGrid2D::BugItHere()
+{
+	FVector WorldLocation = FVector(MouseCursorPosWorld, 0);
+
+	FHitResult HitResult;
+	const bool bHitResultValid = HitTestZFromLocation(World, MouseCursorPosWorld, HitResult);
+
+	if (bHitResultValid && GCurrentLevelEditingViewportClient)
+	{
+		const FRotator ViewRotation = GCurrentLevelEditingViewportClient->GetViewRotation();
+		const FString GoString = FString::Printf(TEXT("BugItGo %f %f %f %f %f %f"), MouseCursorPosWorld.X, MouseCursorPosWorld.Y, HitResult.Location.Z + 1000.0f, ViewRotation.Pitch, ViewRotation.Yaw, ViewRotation.Roll);
+		
+		UE_LOG(LogWorldPartition, Log, TEXT("%s"), *GoString);
+
+		FPlatformApplicationMisc::ClipboardCopy(*GoString);
+	}
+}
+
 void SWorldPartitionEditorGrid2D::MoveCameraHere()
 {
 	FVector WorldLocation = FVector(MouseCursorPosWorld, 0);
 
 	FHitResult HitResult;
-	const FVector TraceStart(WorldLocation.X, WorldLocation.Y, HALF_WORLD_MAX);
-	const FVector TraceEnd(WorldLocation.X, WorldLocation.Y, -HALF_WORLD_MAX);
-	const FCollisionQueryParams TraceParams(SCENE_QUERY_STAT(SWorldPartitionEditorGrid2D_MoveCameraHere), true);
-	const bool bHitResultValid = World->LineTraceSingleByChannel(HitResult, TraceStart, TraceEnd, ECC_WorldStatic, TraceParams);
+	const bool bHitResultValid = HitTestZFromLocation(World, MouseCursorPosWorld, HitResult);
 
 	for (FLevelEditorViewportClient* LevelVC : GEditor->GetLevelViewportClients())
 	{
@@ -596,10 +624,7 @@ void SWorldPartitionEditorGrid2D::PlayFromHere()
 	FRotator StartRotation;
 
 	FHitResult HitResult;
-	const FVector TraceStart(StartLocation.X, StartLocation.Y, HALF_WORLD_MAX);
-	const FVector TraceEnd(StartLocation.X, StartLocation.Y, -HALF_WORLD_MAX);
-	const FCollisionQueryParams TraceParams(SCENE_QUERY_STAT(SWorldPartitionEditorGrid2D_MoveCameraHere), true);
-	const bool bHitResultValid = World->LineTraceSingleByChannel(HitResult, TraceStart, TraceEnd, ECC_WorldStatic, TraceParams);
+	const bool bHitResultValid = HitTestZFromLocation(World, MouseCursorPosWorld, HitResult);
 
 	FLevelEditorModule& LevelEditorModule = FModuleManager::GetModuleChecked<FLevelEditorModule>("LevelEditor");
 	TSharedPtr<IAssetViewport> ActiveLevelViewport = LevelEditorModule.GetFirstActiveViewport();
@@ -741,6 +766,11 @@ FReply SWorldPartitionEditorGrid2D::OnMouseButtonUp(const FGeometry& MyGeometry,
 			MenuBuilder.BeginSection(NAME_None, LOCTEXT("WorldPartitionMisc", "Misc"));
 				MenuBuilder.AddMenuEntry(Commands.MoveCameraHere);
 
+				if (!GetDefault<UWorldPartitionEditorSettings>()->bDisableBugIt)
+				{
+					MenuBuilder.AddMenuEntry(Commands.BugItHere);
+				}
+
 				if (!GetDefault<UWorldPartitionEditorSettings>()->bDisablePIE)
 				{
 					MenuBuilder.AddMenuEntry(Commands.PlayFromHere);
diff --git a/Engine/Source/Editor/WorldPartitionEditor/Private/WorldPartitionEditorModule.cpp b/Engine/Source/Editor/WorldPartitionEditor/Private/WorldPartitionEditorModule.cpp
--- a/Engine/Source/Editor/WorldPartitionEditor/Private/WorldPartitionEditorModule.cpp
+++ b/Engine/Source/Editor/WorldPartitionEditor/Private/WorldPartitionEditorModule.cpp
@@ -259,6 +259,16 @@ void FWorldPartitionEditorModule::SetDisablePIE(bool bInDisablePIE)
 	GetMutableDefault<UWorldPartitionEditorSettings>()->bDisablePIE = bInDisablePIE;
 }
 
+bool FWorldPartitionEditorModule::GetDisableBugIt() const
+{
+	return GetDefault<UWorldPartitionEditorSettings>()->bDisableBugIt;
+}
+
+void FWorldPartitionEditorModule::SetDisableBugIt(bool bInDisableBugIt)
+{
+	GetMutableDefault<UWorldPartitionEditorSettings>()->bDisableBugIt = bInDisableBugIt;
+}
+
 void FWorldPartitionEditorModule::OnConvertMap()
 {
 	IContentBrowserSingleton& ContentBrowserSingleton = FModuleManager::LoadModuleChecked<FContentBrowserModule>("ContentBrowser").Get();
@@ -672,6 +682,7 @@ UWorldPartitionEditorSettings::UWorldPartitionEditorSettings()
 	MinimapLowQualityWorldUnitsPerPixelThreshold = 12800;
 	bDisableLoadingInEditor = false;
 	bDisablePIE = false;
+	bDisableBugIt = false;
 }
 
 FString UWorldPartitionConvertOptions::ToCommandletArgs() const
diff --git a/Engine/Source/Editor/WorldPartitionEditor/Public/WorldPartition/SWorldPartitionEditorGrid2D.h b/Engine/Source/Editor/WorldPartitionEditor/Public/WorldPartition/SWorldPartitionEditorGrid2D.h
--- a/Engine/Source/Editor/WorldPartitionEditor/Public/WorldPartition/SWorldPartitionEditorGrid2D.h
+++ b/Engine/Source/Editor/WorldPartitionEditor/Public/WorldPartition/SWorldPartitionEditorGrid2D.h
@@ -25,6 +25,7 @@ protected:
 		TSharedPtr<FUICommandInfo> UnloadHoveredRegion;
 		TSharedPtr<FUICommandInfo> ConvertSelectedRegionsToActors;
 		TSharedPtr<FUICommandInfo> MoveCameraHere;
+		TSharedPtr<FUICommandInfo> BugItHere;
 		TSharedPtr<FUICommandInfo> PlayFromHere;
 		TSharedPtr<FUICommandInfo> LoadFromHere;
 
@@ -45,6 +46,7 @@ protected:
 	void LoadSelectedRegions();
 	void UnloadSelectedRegions();
 	void ConvertSelectedRegionsToActors();
+	void BugItHere();
 	void MoveCameraHere();
 	void PlayFromHere();
 	void LoadFromHere();
diff --git a/Engine/Source/Editor/WorldPartitionEditor/Public/WorldPartition/WorldPartitionEditorSettings.h b/Engine/Source/Editor/WorldPartitionEditor/Public/WorldPartition/WorldPartitionEditorSettings.h
--- a/Engine/Source/Editor/WorldPartitionEditor/Public/WorldPartition/WorldPartitionEditorSettings.h
+++ b/Engine/Source/Editor/WorldPartitionEditor/Public/WorldPartition/WorldPartitionEditorSettings.h
@@ -25,5 +25,6 @@ public:
 	int32 MinimapLowQualityWorldUnitsPerPixelThreshold;
 
 	bool bDisableLoadingInEditor;
+	bool bDisableBugIt;
 	bool bDisablePIE;
 };
\ No newline at end of file
diff --git a/Engine/Source/Editor/WorldPartitionEditor/Public/WorldPartitionEditorModule.h b/Engine/Source/Editor/WorldPartitionEditor/Public/WorldPartitionEditorModule.h
--- a/Engine/Source/Editor/WorldPartitionEditor/Public/WorldPartitionEditorModule.h
+++ b/Engine/Source/Editor/WorldPartitionEditor/Public/WorldPartitionEditorModule.h
@@ -59,6 +59,12 @@ public:
 	virtual bool GetDisablePIE() const override;
 	virtual void SetDisablePIE(bool bInDisablePIE) override;
 
+	/**
+	 * Returns if bugit command in the editor is disabled or not.
+	 */
+	virtual bool GetDisableBugIt() const override;
+	virtual void SetDisableBugIt(bool bInDisableBugIt) override;
+
 	/**
 	 * Convert the specified map to a world partition map.
 	 */
diff --git a/Engine/Source/Editor/WorldPartitionEditor/WorldPartitionEditor.Build.cs b/Engine/Source/Editor/WorldPartitionEditor/WorldPartitionEditor.Build.cs
--- a/Engine/Source/Editor/WorldPartitionEditor/WorldPartitionEditor.Build.cs
+++ b/Engine/Source/Editor/WorldPartitionEditor/WorldPartitionEditor.Build.cs
@@ -8,6 +8,7 @@ public class WorldPartitionEditor : ModuleRules
     {     
         PrivateDependencyModuleNames.AddRange(
             new string[] {
+				"ApplicationCore",
 				"Core",
 				"CoreUObject",
 				"EditorFramework",
diff --git a/Engine/Source/Runtime/Engine/Public/WorldPartition/IWorldPartitionEditorModule.h b/Engine/Source/Runtime/Engine/Public/WorldPartition/IWorldPartitionEditorModule.h
--- a/Engine/Source/Runtime/Engine/Public/WorldPartition/IWorldPartitionEditorModule.h
+++ b/Engine/Source/Runtime/Engine/Public/WorldPartition/IWorldPartitionEditorModule.h
@@ -41,6 +41,9 @@ public:
 	virtual bool GetDisablePIE() const = 0;
 	virtual void SetDisablePIE(bool bInDisablePIE) = 0;
 
+	virtual bool GetDisableBugIt() const = 0;
+	virtual void SetDisableBugIt(bool bInDisableBugIt) = 0;
+
 	/** Triggered when a world is added. */
 	DECLARE_EVENT_OneParam(IWorldPartitionEditorModule, FWorldPartitionCreated, UWorld*);
 
diff --git a/Engine/Source/Runtime/Engine/Public/WorldPartition/WorldPartitionLog.h b/Engine/Source/Runtime/Engine/Public/WorldPartition/WorldPartitionLog.h
--- a/Engine/Source/Runtime/Engine/Public/WorldPartition/WorldPartitionLog.h
+++ b/Engine/Source/Runtime/Engine/Public/WorldPartition/WorldPartitionLog.h
@@ -3,4 +3,4 @@
 
 #include "CoreMinimal.h"
 
-DECLARE_LOG_CATEGORY_EXTERN(LogWorldPartition, Log, All);
\ No newline at end of file
+ENGINE_API DECLARE_LOG_CATEGORY_EXTERN(LogWorldPartition, Log, All);
\ No newline at end of file
