World Partition Editor:
- Convert hand-made top toolbar to a FToolBarBuilder so it properly handle smaller window sizes by collapsing to a menu instead.
- Switched camera traces to use the camera collision channel.
- Added a transient option to hide the grid display.

#rnx

diff --git a/Engine/Source/Editor/WorldPartitionEditor/Private/WorldPartition/SWorldPartitionEditorGrid2D.cpp b/Engine/Source/Editor/WorldPartitionEditor/Private/WorldPartition/SWorldPartitionEditorGrid2D.cpp
--- a/Engine/Source/Editor/WorldPartitionEditor/Private/WorldPartition/SWorldPartitionEditorGrid2D.cpp
+++ b/Engine/Source/Editor/WorldPartitionEditor/Private/WorldPartition/SWorldPartitionEditorGrid2D.cpp
@@ -22,6 +22,7 @@
 #include "TextureResource.h"
 #include "Widgets/Input/SButton.h"
 #include "Widgets/Input/SCheckBox.h"
+#include "Framework/MultiBox/MultiBoxBuilder.h"
 #include "WorldBrowserModule.h"
 #include "WorldPartition/LoaderAdapter/LoaderAdapterShape.h"
 #include "WorldPartition/WorldPartition.h"
@@ -161,7 +162,7 @@ static bool HitTestZFromLocation(UWorld* World, const FVector2D& WorldLocation,
 	const FVector TraceStart(WorldLocation.X, WorldLocation.Y, HALF_WORLD_MAX);
 	const FVector TraceEnd(WorldLocation.X, WorldLocation.Y, -HALF_WORLD_MAX);
 	const FCollisionQueryParams TraceParams(SCENE_QUERY_STAT(SWorldPartitionEditorGrid2D_HitTestZFromLocation), true);
-	return World->LineTraceSingleByChannel(OutResult, TraceStart, TraceEnd, ECC_WorldStatic, TraceParams);
+	return World->LineTraceSingleByChannel(OutResult, TraceStart, TraceEnd, ECC_Camera, TraceParams);
 }
 
 template <class T>
@@ -250,9 +251,9 @@ void SWorldPartitionEditorGrid2D::FEditorCommands::RegisterCommands()
 	UI_COMMAND(UnloadSelectedRegions, "Unload Selected Regions", "Unload the selected regions.", EUserInterfaceActionType::Button, FInputChord());
 	UI_COMMAND(ConvertSelectedRegionsToActors, "Convert Selected Regions To Actors", "Convert the selected regions to actors.", EUserInterfaceActionType::Button, FInputChord());
 	UI_COMMAND(MoveCameraHere, "Move Camera Here", "Move the camera to the selected location.", EUserInterfaceActionType::Button, FInputChord());
-	UI_COMMAND(BugItHere, "Bug It", "Log BugItGo command of the selected location to console (also adds it to clipboard).", EUserInterfaceActionType::Button, FInputChord());
 	UI_COMMAND(PlayFromHere, "Play From Here", "Play from here.", EUserInterfaceActionType::Button, FInputChord());
 	UI_COMMAND(LoadFromHere, "Load From Here", "Load from here.", EUserInterfaceActionType::Button, FInputChord());
+	UI_COMMAND(BugItHere, "Bug It Here", "Log BugItGo command of the selected location to console (also adds it to clipboard).", EUserInterfaceActionType::Button, FInputChord());
 }
 
 SWorldPartitionEditorGrid2D::SWorldPartitionEditorGrid2D()
@@ -265,6 +266,7 @@ SWorldPartitionEditorGrid2D::SWorldPartitionEditorGrid2D()
 	, bIsPanning(false)
 	, bIsMeasuring(false)
 	, bShowActors(false)
+	, bShowGrid(true)
 	, bFollowPlayerInPIE(false)
 	, SelectBox(ForceInit)
 	, SelectBoxGridSnapped(ForceInit)
@@ -296,118 +298,125 @@ void SWorldPartitionEditorGrid2D::Construct(const FArguments& InArgs)
 	Scale = 0.00133333332;
 	TotalMouseDelta = 0;
 
-	// UI
-	ChildSlot
-	[
-		SNew(SOverlay)
+	auto MakeToolBarWidget = [this]()
+	{
+		FSlimHorizontalToolBarBuilder ToolbarBuilder(CommandList, FMultiBoxCustomization::None);
 
-		// Top status bar
-		+SOverlay::Slot()
-		.VAlign(VAlign_Top)
-		[
-			SNew(SBorder)
-			.BorderImage(FAppStyle::GetBrush(TEXT("Graph.TitleBackground")))
-			[
-				SNew(SVerticalBox)
+		FName ToolBarStyle = "EditorViewportToolBar";
+		ToolbarBuilder.SetStyle(&FAppStyle::Get(), ToolBarStyle);
+		ToolbarBuilder.SetLabelVisibility(EVisibility::Collapsed);
 
-				+SVerticalBox::Slot()
-				.AutoHeight()
-				[
-					SNew(SHorizontalBox)
+		const FEditableTextBoxStyle& TextBoxStyle = FCoreStyle::Get().GetWidgetStyle<FEditableTextBoxStyle>("NormalEditableTextBox");
 
-					+SHorizontalBox::Slot()
-					.AutoWidth()
-					[
-						SNew(SCheckBox)
-						.IsChecked(bShowActors ? ECheckBoxState::Checked : ECheckBoxState::Unchecked)
-						.IsEnabled(true)
-						.OnCheckStateChanged(FOnCheckStateChanged::CreateLambda([=](ECheckBoxState State) { bShowActors = !bShowActors; }))
-					]
-					+SHorizontalBox::Slot()
-					.FillWidth(1.0f)
+		ToolbarBuilder.BeginSection("Toggles");
+		{
+			ToolbarBuilder.AddWidget(SNew(SCheckBox)
+				.ForegroundColor(FSlateColor::UseForeground())
+				.BorderBackgroundColor(TextBoxStyle.BackgroundColor)
+				.IsChecked(bShowActors ? ECheckBoxState::Checked : ECheckBoxState::Unchecked)
+				.OnCheckStateChanged(FOnCheckStateChanged::CreateLambda([=](ECheckBoxState State) { bShowActors = !bShowActors; }))
+				[
+					SNew(SBox)
 					.VAlign(VAlign_Center)
+					.HAlign(HAlign_Center)
+					.Padding(FMargin(4.0f, 2.0f))
 					[
 						SNew(STextBlock)
-						.AutoWrapText(true)
-						.IsEnabled(true)
 						.Text(LOCTEXT("ShowActors", "Show Actors"))
 					]
-					+SHorizontalBox::Slot()
-					.AutoWidth()
+				]);
+
+			ToolbarBuilder.AddWidget(SNew(SCheckBox)
+				.ForegroundColor(FSlateColor::UseForeground())
+				.BorderBackgroundColor(TextBoxStyle.BackgroundColor)
+				.IsChecked(bShowGrid ? ECheckBoxState::Checked : ECheckBoxState::Unchecked)
+				.OnCheckStateChanged(FOnCheckStateChanged::CreateLambda([=](ECheckBoxState State) { bShowGrid = !bShowGrid; }))
+				[
+					SNew(SBox)
+					.VAlign(VAlign_Center)
+					.HAlign(HAlign_Center)
+					.Padding(FMargin(4.0f, 2.0f))
 					[
-						SNew(SCheckBox)
-						.IsChecked(bFollowPlayerInPIE ? ECheckBoxState::Checked : ECheckBoxState::Unchecked)
-						.IsEnabled(true)
-						.Visibility_Lambda([this]() { return GetDefault<UWorldPartitionEditorSettings>()->bDisablePIE ? EVisibility::Hidden : EVisibility::Visible; })
-						.OnCheckStateChanged(FOnCheckStateChanged::CreateLambda([=](ECheckBoxState State) { bFollowPlayerInPIE = !bFollowPlayerInPIE; }))
+						SNew(STextBlock)
+						.Text(LOCTEXT("ShowGrid", "Show Grid"))
 					]
-					+SHorizontalBox::Slot()
-					.FillWidth(1.0f)
+				]);
+
+			ToolbarBuilder.AddWidget(SNew(SCheckBox)
+				.ForegroundColor(FSlateColor::UseForeground())
+				.IsChecked(bFollowPlayerInPIE ? ECheckBoxState::Checked : ECheckBoxState::Unchecked)
+				.Visibility_Lambda([this]() { return GetDefault<UWorldPartitionEditorSettings>()->bDisablePIE ? EVisibility::Hidden : EVisibility::Visible; })
+				.OnCheckStateChanged(FOnCheckStateChanged::CreateLambda([=](ECheckBoxState State) { bFollowPlayerInPIE = !bFollowPlayerInPIE; }))
+				[
+					SNew(SBox)
 					.VAlign(VAlign_Center)
+					.HAlign(HAlign_Center)
+					.Padding(FMargin(4.0f, 2.0f))
 					[
 						SNew(STextBlock)
-						.AutoWrapText(true)
-						.IsEnabled(true)
-						.Visibility_Lambda([this]() { return GetDefault<UWorldPartitionEditorSettings>()->bDisablePIE ? EVisibility::Hidden : EVisibility::Visible; })
 						.Text(LOCTEXT("FollowPlayerInPIE", "Follow Player in PIE"))
 					]
-					+SHorizontalBox::Slot()
-					.AutoWidth()
-					[
-						SNew(SCheckBox)
-						.IsChecked(GetMutableDefault<UWorldPartitionEditorPerProjectUserSettings>()->GetBugItGoLoadRegion() ? ECheckBoxState::Checked : ECheckBoxState::Unchecked)
-						.IsEnabled(true)
-						.Visibility_Lambda([this]() { return GetDefault<UWorldPartitionEditorSettings>()->bDisableLoadingInEditor ? EVisibility::Hidden : EVisibility::Visible; })
-						.OnCheckStateChanged(FOnCheckStateChanged::CreateLambda([=](ECheckBoxState State) { GetMutableDefault<UWorldPartitionEditorPerProjectUserSettings>()->SetBugItGoLoadRegion(State == ECheckBoxState::Checked); }))
-					]
-					+SHorizontalBox::Slot()
-					.FillWidth(1.0f)
+				]);
+
+			ToolbarBuilder.AddWidget(SNew(SCheckBox)
+				.ForegroundColor(FSlateColor::UseForeground())
+				.IsChecked(GetMutableDefault<UWorldPartitionEditorPerProjectUserSettings>()->GetBugItGoLoadRegion() ? ECheckBoxState::Checked : ECheckBoxState::Unchecked)
+				.Visibility_Lambda([this]() { return GetDefault<UWorldPartitionEditorSettings>()->bDisableLoadingInEditor ? EVisibility::Hidden : EVisibility::Visible; })
+				.OnCheckStateChanged(FOnCheckStateChanged::CreateLambda([=](ECheckBoxState State) { GetMutableDefault<UWorldPartitionEditorPerProjectUserSettings>()->SetBugItGoLoadRegion(State == ECheckBoxState::Checked); }))
+				[
+					SNew(SBox)
 					.VAlign(VAlign_Center)
+					.HAlign(HAlign_Center)
+					.Padding(FMargin(4.0f, 2.0f))
 					[
 						SNew(STextBlock)
-						.AutoWrapText(true)
-						.IsEnabled(true)
-						.Visibility_Lambda([this]() { return GetDefault<UWorldPartitionEditorSettings>()->bDisableLoadingInEditor ? EVisibility::Hidden : EVisibility::Visible; })
 						.Text(LOCTEXT("BugItGoLoadRegion", "BugItGo Load Region"))
 					]
-					+SHorizontalBox::Slot()
-					.AutoWidth()
-					[
-						SNew(SCheckBox)
-						.IsChecked(GetMutableDefault<UWorldPartitionEditorPerProjectUserSettings>()->GetShowCellCoords() ? ECheckBoxState::Checked : ECheckBoxState::Unchecked)
-						.IsEnabled(true)
-						.Visibility_Lambda([this]() { return WorldPartition->IsStreamingEnabled() ? EVisibility::Visible : EVisibility::Hidden; })
-						.OnCheckStateChanged(FOnCheckStateChanged::CreateLambda([=](ECheckBoxState State) { GetMutableDefault<UWorldPartitionEditorPerProjectUserSettings>()->SetShowCellCoords(State == ECheckBoxState::Checked); }))
-					]
-					+SHorizontalBox::Slot()
-					.FillWidth(1.0f)
+				]);
+
+			ToolbarBuilder.AddWidget(SNew(SCheckBox)
+				.ForegroundColor(FSlateColor::UseForeground())
+				.IsChecked(GetMutableDefault<UWorldPartitionEditorPerProjectUserSettings>()->GetShowCellCoords() ? ECheckBoxState::Checked : ECheckBoxState::Unchecked)
+				.Visibility_Lambda([this]() { return WorldPartition->IsStreamingEnabled() ? EVisibility::Visible : EVisibility::Hidden; })
+				.OnCheckStateChanged(FOnCheckStateChanged::CreateLambda([=](ECheckBoxState State) { GetMutableDefault<UWorldPartitionEditorPerProjectUserSettings>()->SetShowCellCoords(State == ECheckBoxState::Checked); }))
+				[
+					SNew(SBox)
 					.VAlign(VAlign_Center)
+					.HAlign(HAlign_Center)
+					.Padding(FMargin(4.0f, 2.0f))
 					[
 						SNew(STextBlock)
-						.AutoWrapText(true)
-						.IsEnabled(true)
-						.Visibility_Lambda([this]() { return WorldPartition->IsStreamingEnabled() ? EVisibility::Visible : EVisibility::Hidden; })
 						.Text(LOCTEXT("ShowCellCoords", "Show Cell Coords"))
 					]
-					+ SHorizontalBox::Slot()
-					.AutoWidth()
-					[
-						SNew(SButton)
-						.Text_Lambda([this]() { return GEditor->GetSelectedActors()->Num() ? LOCTEXT("FocusSelection", "Focus Selection") : LOCTEXT("FocusWorld", "Focus World"); })
-						.OnClicked(this, &SWorldPartitionEditorGrid2D::FocusSelection)
-						.IsEnabled_Lambda([this]() { return IsInteractive(); })
-					]
-					+ SHorizontalBox::Slot()
-					.AutoWidth()
-					[
-						SNew(SButton)
-						.Text(LOCTEXT("FocusLoadedRegions", "Focus Loaded Regions"))
-						.OnClicked(this, &SWorldPartitionEditorGrid2D::FocusLoadedRegions)
-						.IsEnabled_Lambda([this]() { return IsInteractive() && WorldPartition && WorldPartition->HasLoadedUserCreatedRegions(); })
-						.Visibility_Lambda([this]() { return GetDefault<UWorldPartitionEditorSettings>()->bDisableLoadingInEditor ? EVisibility::Hidden : EVisibility::Visible; })
-					]
-				]
-			]
+				]);
+
+			ToolbarBuilder.AddWidget(SNew(SButton)
+				.ForegroundColor(FSlateColor::UseForeground())
+				.Text_Lambda([this]() { return GEditor->GetSelectedActors()->Num() ? LOCTEXT("FocusSelection", "Focus Selection") : LOCTEXT("FocusWorld", "Focus World"); })
+				.OnClicked(this, &SWorldPartitionEditorGrid2D::FocusSelection)
+				.IsEnabled_Lambda([this]() { return IsInteractive(); }));
+
+			ToolbarBuilder.AddWidget(SNew(SButton)
+				.ForegroundColor(FSlateColor::UseForeground())
+				.Text(LOCTEXT("FocusLoadedRegions", "Focus Loaded Regions"))
+				.OnClicked(this, &SWorldPartitionEditorGrid2D::FocusLoadedRegions)
+				.IsEnabled_Lambda([this]() { return IsInteractive() && WorldPartition && WorldPartition->HasLoadedUserCreatedRegions(); })
+				.Visibility_Lambda([this]() { return GetDefault<UWorldPartitionEditorSettings>()->bDisableLoadingInEditor ? EVisibility::Hidden : EVisibility::Visible; }));
+		}
+
+		return ToolbarBuilder.MakeWidget();
+	};
+
+	// UI
+	ChildSlot
+	[
+		SNew(SOverlay)
+
+		// Top status bar
+		+SOverlay::Slot()
+		.VAlign(VAlign_Top)
+		[
+			MakeToolBarWidget()
 		]
 		+SOverlay::Slot()
 		.VAlign(VAlign_Bottom)
@@ -467,9 +476,9 @@ void SWorldPartitionEditorGrid2D::Construct(const FArguments& InArgs)
 	ActionList.MapAction(Commands.UnloadSelectedRegions, FExecuteAction::CreateSP(this, &SWorldPartitionEditorGrid2D::UnloadSelectedRegions), FCanExecuteAction::CreateLambda(CanUnloadSelectedRegions));
 	ActionList.MapAction(Commands.ConvertSelectedRegionsToActors, FExecuteAction::CreateSP(this, &SWorldPartitionEditorGrid2D::ConvertSelectedRegionsToActors), FCanExecuteAction::CreateLambda(CanConvertSelectedRegionsToActors));
 	ActionList.MapAction(Commands.MoveCameraHere, FExecuteAction::CreateSP(this, &SWorldPartitionEditorGrid2D::MoveCameraHere));
-	ActionList.MapAction(Commands.BugItHere, FExecuteAction::CreateSP(this, &SWorldPartitionEditorGrid2D::BugItHere));
 	ActionList.MapAction(Commands.PlayFromHere, FExecuteAction::CreateSP(this, &SWorldPartitionEditorGrid2D::PlayFromHere));
 	ActionList.MapAction(Commands.LoadFromHere, FExecuteAction::CreateSP(this, &SWorldPartitionEditorGrid2D::LoadFromHere));
+	ActionList.MapAction(Commands.BugItHere, FExecuteAction::CreateSP(this, &SWorldPartitionEditorGrid2D::BugItHere));
 }
 
 void SWorldPartitionEditorGrid2D::UpdateWorldMiniMapDetails()
@@ -582,24 +591,6 @@ void SWorldPartitionEditorGrid2D::ConvertSelectedRegionsToActors()
 	Refresh();
 }
 
-void SWorldPartitionEditorGrid2D::BugItHere()
-{
-	FVector WorldLocation = FVector(MouseCursorPosWorld, 0);
-
-	FHitResult HitResult;
-	const bool bHitResultValid = HitTestZFromLocation(World, MouseCursorPosWorld, HitResult);
-
-	if (bHitResultValid && GCurrentLevelEditingViewportClient)
-	{
-		const FRotator ViewRotation = GCurrentLevelEditingViewportClient->GetViewRotation();
-		const FString GoString = FString::Printf(TEXT("BugItGo %f %f %f %f %f %f"), MouseCursorPosWorld.X, MouseCursorPosWorld.Y, HitResult.Location.Z + 1000.0f, ViewRotation.Pitch, ViewRotation.Yaw, ViewRotation.Roll);
-		
-		UE_LOG(LogWorldPartition, Log, TEXT("%s"), *GoString);
-
-		FPlatformApplicationMisc::ClipboardCopy(*GoString);
-	}
-}
-
 void SWorldPartitionEditorGrid2D::MoveCameraHere()
 {
 	FVector WorldLocation = FVector(MouseCursorPosWorld, 0);
@@ -665,6 +656,24 @@ void SWorldPartitionEditorGrid2D::LoadFromHere()
 	EditorLoaderAdapter->GetLoaderAdapter()->Load();
 }
 
+void SWorldPartitionEditorGrid2D::BugItHere()
+{
+	FVector WorldLocation = FVector(MouseCursorPosWorld, 0);
+
+	FHitResult HitResult;
+	const bool bHitResultValid = HitTestZFromLocation(World, MouseCursorPosWorld, HitResult);
+
+	if (bHitResultValid && GCurrentLevelEditingViewportClient)
+	{
+		const FRotator ViewRotation = GCurrentLevelEditingViewportClient->GetViewRotation();
+		const FString GoString = FString::Printf(TEXT("BugItGo %f %f %f %f %f %f"), MouseCursorPosWorld.X, MouseCursorPosWorld.Y, HitResult.Location.Z + 1000.0f, ViewRotation.Pitch, ViewRotation.Yaw, ViewRotation.Roll);
+		
+		UE_LOG(LogWorldPartition, Log, TEXT("%s"), *GoString);
+
+		FPlatformApplicationMisc::ClipboardCopy(*GoString);
+	}
+}
+
 bool SWorldPartitionEditorGrid2D::IsFollowPlayerInPIE() const
 {
 	return bFollowPlayerInPIE && (UWorldPartition::IsSimulating() || GEditor->PlayWorld);
diff --git a/Engine/Source/Editor/WorldPartitionEditor/Private/WorldPartition/SWorldPartitionEditorGridSpatialHash.cpp b/Engine/Source/Editor/WorldPartitionEditor/Private/WorldPartition/SWorldPartitionEditorGridSpatialHash.cpp
--- a/Engine/Source/Editor/WorldPartitionEditor/Private/WorldPartition/SWorldPartitionEditorGridSpatialHash.cpp
+++ b/Engine/Source/Editor/WorldPartitionEditor/Private/WorldPartition/SWorldPartitionEditorGridSpatialHash.cpp
@@ -178,9 +178,12 @@ int32 SWorldPartitionEditorGridSpatialHash::PaintGrid(const FGeometry& AllottedG
 			PaintGridLines(VisibleGridRectWorld, EffectiveCellSize / Granularity, Thickness, FLinearColor(0.1f, 0.1f, 0.1f, Alpha * 0.25f));
 		};
 
-		PaintGridLinesFaded(4, 1.0f);
-		PaintGridLinesFaded(2, 1.0f);
-		PaintGridLinesFaded(1, 2.0f);
+		if (bShowGrid)
+		{
+			PaintGridLinesFaded(4, 1.0f);
+			PaintGridLinesFaded(2, 1.0f);
+			PaintGridLinesFaded(1, 2.0f);
+		}
 
 		// Draw coordinates
 		if ((EffectiveCellSize == EditorSpatialHash->CellSize) && GetMutableDefault<UWorldPartitionEditorPerProjectUserSettings>()->GetShowCellCoords())
diff --git a/Engine/Source/Editor/WorldPartitionEditor/Public/WorldPartition/SWorldPartitionEditorGrid2D.h b/Engine/Source/Editor/WorldPartitionEditor/Public/WorldPartition/SWorldPartitionEditorGrid2D.h
--- a/Engine/Source/Editor/WorldPartitionEditor/Public/WorldPartition/SWorldPartitionEditorGrid2D.h
+++ b/Engine/Source/Editor/WorldPartitionEditor/Public/WorldPartition/SWorldPartitionEditorGrid2D.h
@@ -25,9 +25,9 @@ protected:
 		TSharedPtr<FUICommandInfo> UnloadHoveredRegion;
 		TSharedPtr<FUICommandInfo> ConvertSelectedRegionsToActors;
 		TSharedPtr<FUICommandInfo> MoveCameraHere;
-		TSharedPtr<FUICommandInfo> BugItHere;
 		TSharedPtr<FUICommandInfo> PlayFromHere;
 		TSharedPtr<FUICommandInfo> LoadFromHere;
+		TSharedPtr<FUICommandInfo> BugItHere;
 
 		/**
 		 * Initialize commands
@@ -46,10 +46,11 @@ protected:
 	void LoadSelectedRegions();
 	void UnloadSelectedRegions();
 	void ConvertSelectedRegionsToActors();
-	void BugItHere();
+	
 	void MoveCameraHere();
 	void PlayFromHere();
 	void LoadFromHere();
+	void BugItHere();
 
 	bool IsFollowPlayerInPIE() const;
 	bool IsInteractive() const;
@@ -110,6 +111,7 @@ protected:
 	bool bIsPanning;
 	bool bIsMeasuring;
 	bool bShowActors;
+	bool bShowGrid;
 	bool bFollowPlayerInPIE;
 	FVector2D MouseCursorPos;
 	FVector2D MouseCursorPosWorld;
