World Partition - HLOD: Added cancellable progress display for the Build HLOD button

diff --git a/Engine/Source/Editor/WorldPartitionEditor/Private/WorldPartition/Customizations/WorldPartitionHLODDetailsCustomization.cpp b/Engine/Source/Editor/WorldPartitionEditor/Private/WorldPartition/Customizations/WorldPartitionHLODDetailsCustomization.cpp
--- a/Engine/Source/Editor/WorldPartitionEditor/Private/WorldPartition/Customizations/WorldPartitionHLODDetailsCustomization.cpp
+++ b/Engine/Source/Editor/WorldPartitionEditor/Private/WorldPartition/Customizations/WorldPartitionHLODDetailsCustomization.cpp
@@ -1,24 +1,24 @@
 // Copyright Epic Games, Inc. All Rights Reserved.
 
 #include "WorldPartitionHLODDetailsCustomization.h"
+
+#include "Algo/ForEach.h"
+#include "Algo/Transform.h"
+#include "ComponentRecreateRenderStateContext.h"
 #include "Components/StaticMeshComponent.h"
-#include "MaterialShared.h"
-#include "Widgets/DeclarativeSyntaxSupport.h"
-#include "Widgets/Input/SButton.h"
+#include "DetailCategoryBuilder.h"
+#include "DetailLayoutBuilder.h"
+#include "DetailWidgetRow.h"
 #include "Editor.h"
 #include "Engine/Selection.h"
+#include "MaterialShared.h"
+#include "Misc/ScopedSlowTask.h"
 #include "PropertyHandle.h"
-#include "DetailLayoutBuilder.h"
-#include "DetailCategoryBuilder.h"
-#include "DetailWidgetRow.h"
+#include "Widgets/DeclarativeSyntaxSupport.h"
+#include "Widgets/Input/SButton.h"
 #include "WorldPartition/WorldPartition.h"
 #include "WorldPartition/HLOD/HLODActor.h"
 
-#include "Algo/Transform.h"
-#include "Algo/ForEach.h"
-
-#include "ComponentRecreateRenderStateContext.h"
-
 #define LOCTEXT_NAMESPACE "FWorldPartitionHLODDetailsCustomization"
 
 TSharedRef<IDetailCustomization> FWorldPartitionHLODDetailsCustomization::MakeInstance()
@@ -87,9 +87,19 @@ FReply FWorldPartitionHLODDetailsCustomization::OnBuildHLOD()
 		}
 	}
 
+	FScopedSlowTask Progress(SelectedHLODActors.Num(), LOCTEXT("BuildingHLODsforActors", "Building HLODs for actors..."));
+	Progress.MakeDialog(true);
+
 	// Build HLODs
-	for (AWorldPartitionHLOD* HLODActor : SelectedHLODActors)
+	for (int32 Index = 0; Index < SelectedHLODActors.Num(); ++Index)
 	{
+		Progress.EnterProgressFrame(1.0f, FText::Format(LOCTEXT("BuildingHLODsforActorProgress", "Building HLODs ({0}/{1})"), Index+1, SelectedHLODActors.Num()));
+		if (Progress.ShouldCancel())
+		{
+			break;
+		}
+
+		AWorldPartitionHLOD* HLODActor = SelectedHLODActors[Index];
 		HLODActor->BuildHLOD(true);
 	};
 	
