World Partition:
- Optimize dirty actors detection by using dirty actors list maintained by world partition instead of using an actor iterator.
- Maintain a list of selected actor descriptors to be able to show selected ones in the world partition editor.
- Fixed selected actors labels not faded.

#rnx

diff --git a/Engine/Source/Editor/SceneOutliner/Private/ActorBrowsingMode.cpp b/Engine/Source/Editor/SceneOutliner/Private/ActorBrowsingMode.cpp
--- a/Engine/Source/Editor/SceneOutliner/Private/ActorBrowsingMode.cpp
+++ b/Engine/Source/Editor/SceneOutliner/Private/ActorBrowsingMode.cpp
@@ -12,6 +12,7 @@
 #include "UnrealEdGlobals.h"
 #include "ToolMenus.h"
 #include "Engine/Selection.h"
+#include "Engine/World.h"
 #include "Editor/UnrealEdEngine.h"
 #include "HAL/PlatformApplicationMisc.h"
 #include "HAL/FileManager.h"
@@ -37,6 +38,7 @@
 #include "LevelInstance/LevelInstanceSubsystem.h"
 #include "LevelInstance/LevelInstanceEditorInstanceActor.h"
 #include "WorldPartition/WorldPartition.h"
+#include "WorldPartition/WorldPartitionSubsystem.h"
 #include "ISourceControlModule.h"
 #include "ISourceControlProvider.h"
 #include "SourceControlOperations.h"
@@ -807,11 +809,13 @@ void FActorBrowsingMode::OnItemSelectionChanged(FSceneOutlinerTreeItemPtr TreeIt
 {
 	TArray<AActor*> SelectedActors = Selection.GetData<AActor*>(SceneOutliner::FActorSelector());
 
+	SynchronizeSelectedActorDescs();
+
 	bool bChanged = false;
 	bool bAnyInPIE = false;
 	for (auto* Actor : SelectedActors)
 	{
-		if (!bAnyInPIE && Actor && Actor->GetOutermost()->HasAnyPackageFlags(PKG_PlayInEditor))
+		if (!bAnyInPIE && Actor && Actor->GetPackage()->HasAnyPackageFlags(PKG_PlayInEditor))
 		{
 			bAnyInPIE = true;
 		}
@@ -825,7 +829,7 @@ void FActorBrowsingMode::OnItemSelectionChanged(FSceneOutlinerTreeItemPtr TreeIt
 	for (FSelectionIterator SelectionIt(*GEditor->GetSelectedActors()); SelectionIt && !bChanged; ++SelectionIt)
 	{
 		const AActor* Actor = CastChecked< AActor >(*SelectionIt);
-		if (!bAnyInPIE && Actor->GetOutermost()->HasAnyPackageFlags(PKG_PlayInEditor))
+		if (!bAnyInPIE && Actor->GetPackage()->HasAnyPackageFlags(PKG_PlayInEditor))
 		{
 			bAnyInPIE = true;
 		}
@@ -1278,6 +1282,21 @@ FFolder FActorBrowsingMode::GetWorldDefaultRootFolder() const
 	return FFolder::GetWorldRootFolder(RepresentingWorld.Get());
 }
 
+void FActorBrowsingMode::SynchronizeSelectedActorDescs()
+{
+	if (UWorldPartitionSubsystem* WorldPartitionSubsystem = UWorld::GetSubsystem<UWorldPartitionSubsystem>(RepresentingWorld.Get()))
+	{
+		const FSceneOutlinerItemSelection Selection = SceneOutliner->GetSelection();
+		TArray<FWorldPartitionActorDesc*> SelectedActorDescs = Selection.GetData<FWorldPartitionActorDesc*>(SceneOutliner::FActorDescSelector());
+
+		WorldPartitionSubsystem->SelectedActorDescs.Empty();
+		for (FWorldPartitionActorDesc* SelectedActorDesc : SelectedActorDescs)
+		{
+			WorldPartitionSubsystem->SelectedActorDescs.Add(SelectedActorDesc);
+		}
+	}
+}
+
 FSceneOutlinerDragValidationInfo FActorBrowsingMode::ValidateDrop(const ISceneOutlinerTreeItem& DropTarget, const FSceneOutlinerDragDropPayload& Payload) const
 {
 	if (Payload.Has<FFolderTreeItem>())
@@ -2121,6 +2140,12 @@ void FActorBrowsingMode::UnpinSelectedItems()
 	}
 }
 
+void FActorBrowsingMode::SynchronizeSelection()
+{
+	FActorModeInteractive::SynchronizeSelection();
+	SynchronizeSelectedActorDescs();
+}
+
 FCreateSceneOutlinerMode FActorBrowsingMode::CreateFolderPickerMode(const FFolder::FRootObject& InRootObject) const
 {
 	auto MoveSelectionTo = [this, InRootObject](const FSceneOutlinerTreeItemRef& NewParent)
diff --git a/Engine/Source/Editor/SceneOutliner/Private/ActorMode.cpp b/Engine/Source/Editor/SceneOutliner/Private/ActorMode.cpp
--- a/Engine/Source/Editor/SceneOutliner/Private/ActorMode.cpp
+++ b/Engine/Source/Editor/SceneOutliner/Private/ActorMode.cpp
@@ -84,6 +84,23 @@ namespace SceneOutliner
 
 		return false;
 	}
+
+	bool FActorDescSelector::operator()(const TWeakPtr<ISceneOutlinerTreeItem>& Item, FWorldPartitionActorDesc*& ActorDescPtrOut) const
+	{
+		if (TSharedPtr<ISceneOutlinerTreeItem> ItemPtr = Item.Pin())
+		{
+			if (FActorDescTreeItem* ActorDescItem = ItemPtr->CastTo<FActorDescTreeItem>())
+			{
+				if (FWorldPartitionActorDesc* ActorDesc = ActorDescItem->ActorDescHandle.Get())
+				{
+					ActorDescPtrOut = ActorDesc;
+					return true;
+				}
+			}
+		}
+
+		return false;
+	}
 }
 
 FActorMode::FActorMode(const FActorModeParams& Params)
diff --git a/Engine/Source/Editor/SceneOutliner/Public/ActorBrowsingMode.h b/Engine/Source/Editor/SceneOutliner/Public/ActorBrowsingMode.h
--- a/Engine/Source/Editor/SceneOutliner/Public/ActorBrowsingMode.h
+++ b/Engine/Source/Editor/SceneOutliner/Public/ActorBrowsingMode.h
@@ -57,6 +57,8 @@ public:
 	virtual void UnpinItems(const TArray<FSceneOutlinerTreeItemPtr>& InItems) override;
 	virtual void PinSelectedItems() override;
 	virtual void UnpinSelectedItems() override;
+	virtual void SynchronizeSelection() override;
+
 	/* End ISceneOutlinerMode Interface */
 public:
 	/* External events this mode must respond to */
@@ -114,6 +116,8 @@ private:
 	bool GetFolderNamesFromPayload(const FSceneOutlinerDragDropPayload& InPayload, TArray<FName>& OutFolders, FFolder::FRootObject& OutCommonRootObject) const;
 	FFolder GetWorldDefaultRootFolder() const;
 
+	void SynchronizeSelectedActorDescs();
+
 	/** Filter factories */
 	static TSharedRef<FSceneOutlinerFilter> CreateShowOnlySelectedActorsFilter();
 	static TSharedRef<FSceneOutlinerFilter> CreateHideTemporaryActorsFilter();
diff --git a/Engine/Source/Editor/SceneOutliner/Public/ActorMode.h b/Engine/Source/Editor/SceneOutliner/Public/ActorMode.h
--- a/Engine/Source/Editor/SceneOutliner/Public/ActorMode.h
+++ b/Engine/Source/Editor/SceneOutliner/Public/ActorMode.h
@@ -17,6 +17,12 @@ namespace SceneOutliner
 	{
 		bool operator()(const TWeakPtr<ISceneOutlinerTreeItem>& Item, AActor*& ActorPtrOut) const;
 	};
+
+	/** Functor which can be used to get actor descriptors from a selection  */
+	struct FActorDescSelector
+	{
+		bool operator()(const TWeakPtr<ISceneOutlinerTreeItem>& Item, FWorldPartitionActorDesc*& ActorDescPtrOut) const;
+	};
 }
 
 struct SCENEOUTLINER_API FActorModeParams
diff --git a/Engine/Source/Editor/WorldPartitionEditor/Private/WorldPartition/SWorldPartitionEditorGrid2D.cpp b/Engine/Source/Editor/WorldPartitionEditor/Private/WorldPartition/SWorldPartitionEditorGrid2D.cpp
--- a/Engine/Source/Editor/WorldPartitionEditor/Private/WorldPartition/SWorldPartitionEditorGrid2D.cpp
+++ b/Engine/Source/Editor/WorldPartitionEditor/Private/WorldPartition/SWorldPartitionEditorGrid2D.cpp
@@ -37,6 +37,7 @@
 #include "WorldPartition/WorldPartitionMiniMap.h"
 #include "WorldPartition/WorldPartitionMiniMapHelper.h"
 #include "WorldPartition/WorldPartitionMiniMapVolume.h"
+#include "WorldPartition/WorldPartitionSubsystem.h"
 #include "HAL/PlatformApplicationMisc.h"
 
 #include "IAssetViewport.h"
@@ -392,7 +393,11 @@ void SWorldPartitionEditorGrid2D::Construct(const FArguments& InArgs)
 
 			ToolbarBuilder.AddWidget(SNew(SButton)
 				.ForegroundColor(FSlateColor::UseForeground())
-				.Text_Lambda([this]() { return GEditor->GetSelectedActors()->Num() ? LOCTEXT("FocusSelection", "Focus Selection") : LOCTEXT("FocusWorld", "Focus World"); })
+				.Text_Lambda([this]()
+				{
+					UWorldPartitionSubsystem* WorldPartitionSubsystem = UWorld::GetSubsystem<UWorldPartitionSubsystem>(World);
+					return (GEditor->GetSelectedActors()->Num() || WorldPartitionSubsystem->SelectedActorDescs.Num()) ? LOCTEXT("FocusSelection", "Focus Selection") : LOCTEXT("FocusWorld", "Focus World");
+				})
 				.OnClicked(this, &SWorldPartitionEditorGrid2D::FocusSelection)
 				.IsEnabled_Lambda([this]() { return IsInteractive(); }));
 
@@ -1065,23 +1070,17 @@ void SWorldPartitionEditorGrid2D::Tick(const FGeometry& AllottedGeometry, const
 
 	// Also include transient actor loader adapters that might have been spawned by blutilities, etc. Since these actors can't be saved because they are transient,
 	// they will never get an actor descriptor so they will never appear in the world partition editor. Also include unsaved, newly created actors for convenience.
-	for (TActorIterator<AActor> ActorIt(World); ActorIt; ++ActorIt)
+	for (auto& [Reference, Actor] : WorldPartition->GetDirtyActors())
 	{
-		if (AActor* Actor = *ActorIt)
+		if (!WorldPartition->GetActorDesc(Actor->GetActorGuid()) && Actor->Implements<UWorldPartitionActorLoaderInterface>())
 		{
-			if (!WorldPartition->GetActorDesc(Actor->GetActorGuid()) && Actor->Implements<UWorldPartitionActorLoaderInterface>())
-			{
-				if (IWorldPartitionActorLoaderInterface::ILoaderAdapter* LoaderAdapter = Cast<IWorldPartitionActorLoaderInterface>(Actor)->GetLoaderAdapter())
-				{
-					ShownLoaderInterfaces.Add(Actor);
-				}
-			}
-
-			if(Actor->IsSelected() || Actor->GetPackage()->IsDirty())
+			if (IWorldPartitionActorLoaderInterface::ILoaderAdapter* LoaderAdapter = Cast<IWorldPartitionActorLoaderInterface>(Actor)->GetLoaderAdapter())
 			{
-				DirtyActorGuids.Add(Actor->GetActorGuid());
+				ShownLoaderInterfaces.Add(Actor);
 			}
 		}
+
+		DirtyActorGuids.Add(Actor->GetActorGuid());
 	}
 
 	FLoaderInterfaceSet LastHoveredLoaderInterfaces = MoveTemp(HoveredLoaderInterfaces);
@@ -1152,6 +1151,13 @@ void SWorldPartitionEditorGrid2D::Tick(const FGeometry& AllottedGeometry, const
 			ShownActorGuids.Add(Actor->GetActorGuid());
 		}
 	}
+
+	// Include selected actor descriptors
+	UWorldPartitionSubsystem* WorldPartitionSubsystem = UWorld::GetSubsystem<UWorldPartitionSubsystem>(World);
+	for (FWorldPartitionActorDesc* SelectedActorDesc : WorldPartitionSubsystem->SelectedActorDescs)
+	{
+		ShownActorGuids.Add(SelectedActorDesc->GetGuid());
+	}
 }
 
 uint32 SWorldPartitionEditorGrid2D::PaintActors(const FGeometry& AllottedGeometry, const FSlateRect& MyCullingRect, FSlateWindowElementList& OutDrawElements, uint32 LayerId) const
@@ -1159,6 +1165,7 @@ uint32 SWorldPartitionEditorGrid2D::PaintActors(const FGeometry& AllottedGeometr
 	const FBox2D ViewRect(FVector2D(ForceInitToZero), AllottedGeometry.GetLocalSize());
 	const FBox2D WorldViewRect(ScreenToWorld.TransformPoint(ViewRect.Min), ScreenToWorld.TransformPoint(ViewRect.Max));
 	const FBox ViewRectWorld(FVector(WorldViewRect.Min.X, WorldViewRect.Min.Y, -HALF_WORLD_MAX), FVector(WorldViewRect.Max.X, WorldViewRect.Max.Y, HALF_WORLD_MAX));
+	const UWorldPartitionSubsystem* WorldPartitionSubsystem = UWorld::GetSubsystem<UWorldPartitionSubsystem>(World);
 
 	TArray<FWorldPartitionActorDescViewBoundsProxy> ActorDescList;
 	ActorDescList.Reserve(ShownActorGuids.Num());
@@ -1189,7 +1196,7 @@ uint32 SWorldPartitionEditorGrid2D::PaintActors(const FGeometry& AllottedGeometr
 			if (bShowBackground)
 			{
 				const FSlateColorBrush BackgroundBrush(FLinearColor::Black);
-				const FLinearColor LabelBackgroundColor(0, 0, 0, 0.1f);
+				const FLinearColor LabelBackgroundColor = FLinearColor::Black.CopyWithNewOpacity(Color.A * 0.1f);
 
 				FSlateDrawElement::MakeBox(
 					OutDrawElements,
@@ -1201,12 +1208,15 @@ uint32 SWorldPartitionEditorGrid2D::PaintActors(const FGeometry& AllottedGeometr
 				);
 			}
 
-			const float LabelColorGradient = bShowBackground ? 1.0f : FMath::Clamp(Geometry.GetLocalSize().X / LabelTextSize.X - 1.0f, 0.0f, 1.0f);
-
-			if (LabelColorGradient > 0.0f)
-			{
-				LayerId = DrawTextLabel(OutDrawElements, LayerId, AllottedGeometry, Label, Pos, Color.CopyWithNewOpacity(Color.A * LabelColorGradient), Font);
-			}
+			LayerId = DrawTextLabel(
+				OutDrawElements, 
+				++LayerId, 
+				AllottedGeometry, 
+				Label, 
+				Pos, 
+				Color, 
+				Font
+			);
 		}
 	};
 
@@ -1346,7 +1356,7 @@ uint32 SWorldPartitionEditorGrid2D::PaintActors(const FGeometry& AllottedGeometr
 						// Label
 						{
 							const FString ActorLabel = *LoaderAdapter->GetLabel();
-							const FLinearColor LabelColor(1.0f, 1.0f, 1.0f, LoaderColorGradient * FullScreenColorGradient);
+							const FLinearColor LabelColor = FLinearColor::White.CopyWithNewOpacity(LoaderColorGradient * FullScreenColorGradient);
 							DrawActorLabel(ActorLabel, ActorViewBox.GetCenter(), ActorGeometry, LabelColor, SmallLayoutFont, false);
 						}
 					}
@@ -1401,7 +1411,8 @@ uint32 SWorldPartitionEditorGrid2D::PaintActors(const FGeometry& AllottedGeometr
 					{
 						if (!ActorLabel.IsNone())
 						{
-							DrawActorLabel(ActorLabel.ToString(), bIsBoundsEdgeHovered ? MouseCursorPos : ActorViewBox.GetCenter(), ActorGeometry, FLinearColor::Yellow, SmallLayoutFont, true);
+							const FLinearColor LabelColor = FLinearColor::Yellow.CopyWithNewOpacity(ActorColorGradient);
+							DrawActorLabel(ActorLabel.ToString(), !bIsSelected ? MouseCursorPos : ActorViewBox.GetCenter(), ActorGeometry, LabelColor, SmallLayoutFont, true);
 						}
 
 						ActorColor = FLinearColor::Yellow.CopyWithNewOpacity(ActorColor.A);
@@ -1433,7 +1444,7 @@ uint32 SWorldPartitionEditorGrid2D::PaintActors(const FGeometry& AllottedGeometr
 
 			const FBox ActorBounds = ActorDescView.GetBounds();
 			const AActor* Actor = ActorDescView.GetActor();
-			const bool bIsSelected = Actor ? Actor->IsSelected() : false;
+			const bool bIsSelected = Actor ? Actor->IsSelected() : WorldPartitionSubsystem->SelectedActorDescs.Contains(ActorDescView.GetActorDesc());
 			const bool bIsSpatiallyLoaded = ActorDescView.GetIsSpatiallyLoaded();
 			const FName ActorLabel = ActorDescView.GetActorLabel();
 
@@ -1916,8 +1927,8 @@ int32 SWorldPartitionEditorGrid2D::DrawTextLabel(FSlateWindowElementList& OutDra
 	if (LabelTextSize.X > 0)
 	{
 		const FVector2D LabelTextPos = Pos - LabelTextSize * 0.5f;
-		const FLinearColor LabelForegroundColor(1.0f, 1.0f, 1.0f, Color.A);
-		const FLinearColor LabelShadowColor(0, 0, 0, Color.A);
+		const FLinearColor LabelForegroundColor = FLinearColor::White.CopyWithNewOpacity(Color.A);
+		const FLinearColor LabelShadowColor = FLinearColor::Black.CopyWithNewOpacity(Color.A);
 
 		FSlateDrawElement::MakeText(
 			OutDrawElements,
@@ -1947,19 +1958,21 @@ FReply SWorldPartitionEditorGrid2D::FocusSelection()
 {
 	FBox SelectionBox(ForceInit);
 
-	USelection* SelectedActors = GEditor->GetSelectedActors();
-
-	if (SelectedActors->Num())
+	for (FSelectionIterator It(*GEditor->GetSelectedActors()); It; ++It)
 	{
-		for (FSelectionIterator It(*SelectedActors); It; ++It)
+		if (AActor* Actor = Cast<AActor>(*It))
 		{
-			if (AActor* Actor = Cast<AActor>(*It))
-			{
-				SelectionBox += Actor->GetStreamingBounds();
-			}
+			SelectionBox += Actor->GetStreamingBounds();
 		}
 	}
-	else
+
+	UWorldPartitionSubsystem* WorldPartitionSubsystem = UWorld::GetSubsystem<UWorldPartitionSubsystem>(World);
+	for (FWorldPartitionActorDesc* SelectedActorDesc : WorldPartitionSubsystem->SelectedActorDescs)
+	{
+		SelectionBox += SelectedActorDesc->GetBounds();
+	}
+
+	if (!SelectionBox.IsValid)
 	{
 		// Override the minimap bounds if world partition minimap volumes exists
 		for (TActorIterator<AWorldPartitionMiniMapVolume> It(World); It; ++It)
diff --git a/Engine/Source/Runtime/Engine/Public/WorldPartition/WorldPartition.h b/Engine/Source/Runtime/Engine/Public/WorldPartition/WorldPartition.h
--- a/Engine/Source/Runtime/Engine/Public/WorldPartition/WorldPartition.h
+++ b/Engine/Source/Runtime/Engine/Public/WorldPartition/WorldPartition.h
@@ -214,6 +214,8 @@ public:
 	void OnUserCreatedRegionUnloaded() { check(HasLoadedUserCreatedRegions()); NumUserCreatedLoadedRegions--; }
 
 	bool IsEnablingStreamingJustified() const { return bEnablingStreamingJustified; }
+
+	const TMap<FWorldPartitionReference, AActor*>& GetDirtyActors() const { return DirtyActors; }
 #endif
 
 public:
@@ -307,7 +309,7 @@ private:
 
 #if WITH_EDITORONLY_DATA
 	// Default HLOD layer
-	UPROPERTY(EditAnywhere, Category=WorldPartitionSetup, meta = (DisplayName = "Default HLOD Layer", EditCondition="bEnableStreaming", EditConditionHides, HideEditConditionToggle))
+	UPROPERTY(EditAnywhere, Category = WorldPartitionSetup, meta = (DisplayName = "Default HLOD Layer", EditCondition="bEnableStreaming", EditConditionHides, HideEditConditionToggle))
 	TObjectPtr<class UHLODLayer> DefaultHLODLayer;
 
 	TArray<FWorldPartitionReference> LoadedSubobjects;
diff --git a/Engine/Source/Runtime/Engine/Public/WorldPartition/WorldPartitionSubsystem.h b/Engine/Source/Runtime/Engine/Public/WorldPartition/WorldPartitionSubsystem.h
--- a/Engine/Source/Runtime/Engine/Public/WorldPartition/WorldPartitionSubsystem.h
+++ b/Engine/Source/Runtime/Engine/Public/WorldPartition/WorldPartitionSubsystem.h
@@ -61,6 +61,8 @@ public:
 	void ForEachWorldPartition(TFunctionRef<bool(UWorldPartition*)> Func);
 
 	static bool IsRunningConvertWorldPartitionCommandlet();
+
+	TSet<FWorldPartitionActorDesc*> SelectedActorDescs;
 #endif
 
 private:
