Added (Current) on the current/editing Content Bundle (like Data Layers and Folder).

diff --git a/Engine/Source/Editor/WorldPartitionEditor/Private/WorldPartition/ContentBundle/Outliner/ContentBundleTreeItem.cpp b/Engine/Source/Editor/WorldPartitionEditor/Private/WorldPartition/ContentBundle/Outliner/ContentBundleTreeItem.cpp
--- a/Engine/Source/Editor/WorldPartitionEditor/Private/WorldPartition/ContentBundle/Outliner/ContentBundleTreeItem.cpp
+++ b/Engine/Source/Editor/WorldPartitionEditor/Private/WorldPartition/ContentBundle/Outliner/ContentBundleTreeItem.cpp
@@ -41,7 +41,7 @@ FString FContentBundleTreeItem::GetDisplayString() const
 	TSharedPtr<FContentBundleEditor> ContentBundleEditorPin = ContentBundleEditor.Pin();
 	if (ContentBundleEditorPin != nullptr)
 	{
-		return ContentBundleEditorPin->GetDisplayName();
+		return FString::Printf(TEXT("%s%s"), *ContentBundleEditorPin->GetDisplayName(), ContentBundleEditorPin->IsBeingEdited() ? TEXT(" (Current)") : TEXT(""));
 	}
 
 	return LOCTEXT("ContentBundleMissingForTreeItem", "(Invalid Content Bundle)").ToString();
