Fixed crashes when exiting the editor using "-execcmds=QUIT_EDITOR" (only happens is ContentBundle outliner and/or DataLayer Outliner are spawned)

diff --git a/Engine/Source/Editor/DataLayerEditor/Private/DataLayer/DataLayerHierarchy.cpp b/Engine/Source/Editor/DataLayerEditor/Private/DataLayer/DataLayerHierarchy.cpp
--- a/Engine/Source/Editor/DataLayerEditor/Private/DataLayer/DataLayerHierarchy.cpp
+++ b/Engine/Source/Editor/DataLayerEditor/Private/DataLayer/DataLayerHierarchy.cpp
@@ -112,8 +112,11 @@ FDataLayerHierarchy::~FDataLayerHierarchy()
 		}
 	}
 
-	UDataLayerEditorSubsystem::Get()->OnDataLayerChanged().RemoveAll(this);
-	UDataLayerEditorSubsystem::Get()->OnActorDataLayersChanged().RemoveAll(this);
+	if (UDataLayerEditorSubsystem* DataLayerEditorSubSystem = UDataLayerEditorSubsystem::Get())
+	{
+		DataLayerEditorSubSystem->OnDataLayerChanged().RemoveAll(this);
+		DataLayerEditorSubSystem->OnActorDataLayersChanged().RemoveAll(this);
+	}
 	FWorldDelegates::LevelAddedToWorld.RemoveAll(this);
 	FWorldDelegates::LevelRemovedFromWorld.RemoveAll(this);
 }
diff --git a/Engine/Source/Editor/WorldPartitionEditor/Private/WorldPartition/ContentBundle/Outliner/ContentBundleHierarchy.cpp b/Engine/Source/Editor/WorldPartitionEditor/Private/WorldPartition/ContentBundle/Outliner/ContentBundleHierarchy.cpp
--- a/Engine/Source/Editor/WorldPartitionEditor/Private/WorldPartition/ContentBundle/Outliner/ContentBundleHierarchy.cpp
+++ b/Engine/Source/Editor/WorldPartitionEditor/Private/WorldPartition/ContentBundle/Outliner/ContentBundleHierarchy.cpp
@@ -18,9 +18,12 @@ FContentBundleHiearchy::FContentBundleHiearchy(FContentBundleMode* Mode)
 
 FContentBundleHiearchy::~FContentBundleHiearchy()
 {
-	UContentBundleEditorSubsystem::Get()->OnContentBundleChanged().RemoveAll(this);
-	UContentBundleEditorSubsystem::Get()->OnContentBundleAdded().RemoveAll(this);
-	UContentBundleEditorSubsystem::Get()->OnContentBundleRemoved().RemoveAll(this);
+	if (UContentBundleEditorSubsystem* ContentBundleEditorSubsystem = UContentBundleEditorSubsystem::Get())
+	{
+		ContentBundleEditorSubsystem->OnContentBundleChanged().RemoveAll(this);
+		ContentBundleEditorSubsystem->OnContentBundleAdded().RemoveAll(this);
+		ContentBundleEditorSubsystem->OnContentBundleRemoved().RemoveAll(this);
+	}
 }
 
 TUniquePtr<FContentBundleHiearchy> FContentBundleHiearchy::Create(FContentBundleMode* Mode)
