Added "Browse to asset" on a content bundle.

diff --git a/Engine/Source/Editor/WorldPartitionEditor/Private/WorldPartition/ContentBundle/Outliner/ContentBundleMode.cpp b/Engine/Source/Editor/WorldPartitionEditor/Private/WorldPartition/ContentBundle/Outliner/ContentBundleMode.cpp
--- a/Engine/Source/Editor/WorldPartitionEditor/Private/WorldPartition/ContentBundle/Outliner/ContentBundleMode.cpp
+++ b/Engine/Source/Editor/WorldPartitionEditor/Private/WorldPartition/ContentBundle/Outliner/ContentBundleMode.cpp
@@ -1,14 +1,19 @@
 // Copyright Epic Games, Inc. All Rights Reserved.
 
 #include "WorldPartition/ContentBundle/Outliner/ContentBundleMode.h"
-#include "Engine/Level.h"
 #include "WorldPartition/ContentBundle/Outliner/ContentBundleHierarchy.h"
-#include "ToolMenus.h"
-#include "SceneOutlinerMenuContext.h"
+#include "WorldPartition/ContentBundle/Outliner/ContentBundleTreeItem.h"
 #include "WorldPartition/ContentBundle/ContentBundleEditorSubsystem.h"
 #include "WorldPartition/ContentBundle/ContentBundleEditor.h"
-#include "WorldPartition/ContentBundle/Outliner/ContentBundleTreeItem.h"
+#include "WorldPartition/ContentBundle/ContentBundleDescriptor.h"
+#include "Framework/Commands/UIAction.h"
+#include "Framework/Commands/UICommandList.h"
+#include "Toolkits/GlobalEditorCommonCommands.h"
+#include "SceneOutlinerMenuContext.h"
 #include "Algo/Transform.h"
+#include "Engine/DataAsset.h"
+#include "ToolMenus.h"
+#include "Editor.h"
 
 #define LOCTEXT_NAMESPACE "ContentBundle"
 
@@ -21,7 +26,11 @@ namespace ContentBundleOutlinerPrivate
 FContentBundleMode::FContentBundleMode(const FContentBundleModeCreationParams& Params)
 	:ISceneOutlinerMode(Params.SceneOutliner)
 {
-
+	Commands = MakeShareable(new FUICommandList());
+	Commands->MapAction(FGlobalEditorCommonCommands::Get().FindInContentBrowser, FUIAction(
+		FExecuteAction::CreateRaw(this, &FContentBundleMode::FindInContentBrowser),
+		FCanExecuteAction::CreateRaw(this, &FContentBundleMode::CanFindInContentBrowser)
+	));
 }
 
 FContentBundleMode::~FContentBundleMode()
@@ -85,6 +94,11 @@ FReply FContentBundleMode::OnKeyDown(const FKeyEvent& InKeyEvent)
 		return FReply::Handled();
 	}
 
+	if (Commands->ProcessCommandBindings(InKeyEvent))
+	{
+		return FReply::Handled();
+	}
+
 	return FReply::Unhandled();
 }
 
@@ -105,6 +119,33 @@ UWorld* FContentBundleMode::GetEditingWorld() const
 	return nullptr;
 }
 
+void FContentBundleMode::FindInContentBrowser()
+{
+	if (SceneOutliner)
+	{
+		TArray<UObject*> Objects;
+		SceneOutliner->GetSelection().ForEachItem<FContentBundleTreeItem>([&Objects](const FContentBundleTreeItem& Item)
+		{
+			if (TSharedPtr<FContentBundleEditor> ContentBundleEditor = Item.GetContentBundleEditorPin())
+			{
+				if (UDataAsset* DataAsset = ContentBundleEditor->GetDescriptor()->GetTypedOuter<UDataAsset>())
+				{
+					Objects.Add(DataAsset);
+				}
+			}
+		});
+		if (!Objects.IsEmpty())
+		{
+			GEditor->SyncBrowserToObjects(Objects);
+		}
+	}
+}
+
+bool FContentBundleMode::CanFindInContentBrowser() const
+{
+	return SceneOutliner && SceneOutliner->GetSelection().Has<FContentBundleTreeItem>();
+}
+
 void FContentBundleMode::RegisterContextMenu()
 {
 	UToolMenus* ToolMenus = UToolMenus::Get();
@@ -157,6 +198,11 @@ void FContentBundleMode::RegisterContextMenu()
 							));
 				}
 
+				{
+					FToolMenuSection& Section = InMenu->AddSection("AssetOptionsSection", LOCTEXT("AssetOptionsText", "Asset Options"));
+					Section.AddMenuEntryWithCommandList(FGlobalEditorCommonCommands::Get().FindInContentBrowser, Commands);
+				}
+
 				{
 					FToolMenuSection& Section = InMenu->AddSection("ContentBundleLoadingSection", LOCTEXT("Loading", "Loading"));
 
diff --git a/Engine/Source/Editor/WorldPartitionEditor/Public/WorldPartition/ContentBundle/Outliner/ContentBundleMode.h b/Engine/Source/Editor/WorldPartitionEditor/Public/WorldPartition/ContentBundle/Outliner/ContentBundleMode.h
--- a/Engine/Source/Editor/WorldPartitionEditor/Public/WorldPartition/ContentBundle/Outliner/ContentBundleMode.h
+++ b/Engine/Source/Editor/WorldPartitionEditor/Public/WorldPartition/ContentBundle/Outliner/ContentBundleMode.h
@@ -7,6 +7,7 @@
 
 class SSceneOutliner;
 class FContentBundleEditor;
+class FUICommandList;
 
 struct FContentBundleModeCreationParams
 {
@@ -45,4 +46,14 @@ protected:
 
 	void RegisterContextMenu();
 	void UnregisterContextMenu();
+
+private:
+
+	/** Delegate to handle "Find in Content Browser" context menu option */
+	void FindInContentBrowser();
+
+	/** Delegate to handle enabling the "Find in Content Browser" context menu option */
+	bool CanFindInContentBrowser() const;
+
+	TSharedPtr<FUICommandList> Commands;
 };
\ No newline at end of file
