Content Bundle small fixes
- Prevent Content Bundle descriptor to be deleted
- Allow editing of Content Bundle debug color
- Use Content Bundle debug color in Actor Editor context

diff --git a/Engine/Plugins/Experimental/GameFeatures/Source/GameFeatures/Public/GameFeatureAction_AddWPContent.h b/Engine/Plugins/Experimental/GameFeatures/Source/GameFeatures/Public/GameFeatureAction_AddWPContent.h
--- a/Engine/Plugins/Experimental/GameFeatures/Source/GameFeatures/Public/GameFeatureAction_AddWPContent.h
+++ b/Engine/Plugins/Experimental/GameFeatures/Source/GameFeatures/Public/GameFeatureAction_AddWPContent.h
@@ -33,7 +33,7 @@ public:
 	const UContentBundleDescriptor* GetContentBundleDescriptor() const { return ContentBundleDescriptor; }
 
 private:
-	UPROPERTY(EditAnywhere, Category = ContentBundle)
+	UPROPERTY(VisibleAnywhere, Category = ContentBundle)
 	TObjectPtr<UContentBundleDescriptor> ContentBundleDescriptor;
 
 	TSharedPtr<FContentBundleClient> ContentBundleClient;
diff --git a/Engine/Source/Editor/WorldPartitionEditor/Private/WorldPartition/ContentBundle/ContentBundleEditorSubsystem.cpp b/Engine/Source/Editor/WorldPartitionEditor/Private/WorldPartition/ContentBundle/ContentBundleEditorSubsystem.cpp
--- a/Engine/Source/Editor/WorldPartitionEditor/Private/WorldPartition/ContentBundle/ContentBundleEditorSubsystem.cpp
+++ b/Engine/Source/Editor/WorldPartitionEditor/Private/WorldPartition/ContentBundle/ContentBundleEditorSubsystem.cpp
@@ -76,7 +76,7 @@ TSharedRef<SWidget> UContentBundleEditionSubmodule::GetActorEditorContextWidget(
 
 	TSharedRef<SVerticalBox> OutWidget = SNew(SVerticalBox);
 
-	if (EditingContentBundle != nullptr)
+	if (TSharedPtr<FContentBundleEditor> ContentBundleEditorPin = EditingContentBundle.Pin())
 	{
 		OutWidget->AddSlot().AutoHeight()
 		[
@@ -87,6 +87,14 @@ TSharedRef<SWidget> UContentBundleEditionSubmodule::GetActorEditorContextWidget(
 			.VAlign(VAlign_Center)
 			[
 				SNew(SImage)
+				.ColorAndOpacity_Lambda([this]()
+				{
+					if (TSharedPtr<FContentBundleEditor> ContentBundleEditorPin = EditingContentBundle.Pin())
+					{
+						return ContentBundleEditorPin->GetDebugColor();
+					}
+					return FColor::Black;
+				})
 				.Image(FAppStyle::Get().GetBrush(PLACEHOLDER_ContentBundle_ColorIcon))
 				.DesiredSizeOverride(FVector2D(8, 8))
 			]
@@ -97,12 +105,13 @@ TSharedRef<SWidget> UContentBundleEditionSubmodule::GetActorEditorContextWidget(
 				SNew(STextBlock)
 				.Text_Lambda([this]()
 				{
-					TSharedPtr<FContentBundleEditor> ContentBundleEditorPin = EditingContentBundle.Pin();
-					if (ContentBundleEditorPin != nullptr)
+					if (TSharedPtr<FContentBundleEditor> ContentBundleEditorPin = EditingContentBundle.Pin())
 					{
-						return FText::FromString(ContentBundleEditorPin->GetDisplayName());
+						if (!ContentBundleEditorPin->GetDisplayName().IsEmpty())
+						{
+							return FText::FromString(ContentBundleEditorPin->GetDisplayName());
+						}
 					}
-
 					return FText::FromString(TEXT("Unset Display Name"));
 				})
 					
diff --git a/Engine/Source/Runtime/Engine/Public/WorldPartition/ContentBundle/ContentBundleDescriptor.h b/Engine/Source/Runtime/Engine/Public/WorldPartition/ContentBundle/ContentBundleDescriptor.h
--- a/Engine/Source/Runtime/Engine/Public/WorldPartition/ContentBundle/ContentBundleDescriptor.h
+++ b/Engine/Source/Runtime/Engine/Public/WorldPartition/ContentBundle/ContentBundleDescriptor.h
@@ -44,7 +44,7 @@ private:
 	UPROPERTY(EditAnywhere, Category = BaseInformation)
 	FString DisplayName;
 
-	UPROPERTY(VisibleAnywhere, Category = BaseInformation)
+	UPROPERTY(EditAnywhere, Category = BaseInformation)
 	FColor DebugColor;
 
 	UPROPERTY(VisibleAnywhere, Category = BaseInformation, AdvancedDisplay)
