- Replaced "Edit Content Bundle" by "Make Current Content Bundle" and "Clear Current Content Bundle"
- Added Undo/Redo support on Content Bundle operations
- Added support to Reset Current Content Bundle via Editing Current Context.
- Renamed UContentBundleEditionSubsystem to UContentBundleEditorSubsystem

diff --git a/Engine/Source/Editor/WorldPartitionEditor/Private/WorldPartition/ContentBundle/ContentBundleEditorSubsystem.cpp b/Engine/Source/Editor/WorldPartitionEditor/Private/WorldPartition/ContentBundle/ContentBundleEditorSubsystem.cpp
--- a/Engine/Source/Editor/WorldPartitionEditor/Private/WorldPartition/ContentBundle/ContentBundleEditorSubsystem.cpp
+++ b/Engine/Source/Editor/WorldPartitionEditor/Private/WorldPartition/ContentBundle/ContentBundleEditorSubsystem.cpp
@@ -16,32 +16,60 @@
 #include "Widgets/Text/STextBlock.h"
 #include "Editor.h"
 
-void UContentBundleEditionSubmodule::DoInitialize()
+void UContentBundleEditingSubmodule::DoInitialize()
 {
 	UActorEditorContextSubsystem::Get()->RegisterClient(this);
 }
 
-void UContentBundleEditionSubmodule::DoDenitialize()
+void UContentBundleEditingSubmodule::DoDenitialize()
 {
 	UActorEditorContextSubsystem::Get()->UnregisterClient(this);
 
-	EditingContentBundle = nullptr;
-	ActiveContentBundlesStack.Empty();
+	EditingContentBundleGuid.Invalidate();
+	EditingContentBundlesStack.Empty();
 }
 
-void UContentBundleEditionSubmodule::OnExecuteActorEditorContextAction(UWorld* InWorld, const EActorEditorContextAction& InType, AActor* InActor /* = nullptr */)
+void UContentBundleEditingSubmodule::PreEditUndo()
 {
-	TSharedPtr<FContentBundleEditor> ContentBundleEditorPin = EditingContentBundle.Pin();
+	Super::PreEditUndo();
 
+	PreUndoRedoEditingContentBundleGuid = EditingContentBundleGuid;
+}
+
+void UContentBundleEditingSubmodule::PostEditUndo()
+{
+	Super::PostEditUndo();
+
+	if (PreUndoRedoEditingContentBundleGuid != EditingContentBundleGuid)
+	{
+		if (TSharedPtr<FContentBundleEditor> OldEditingContentBundle = GetEditorContentBundle(PreUndoRedoEditingContentBundleGuid))
+		{
+			StopEditing(OldEditingContentBundle);
+		}
+		if (TSharedPtr<FContentBundleEditor> NewEditingContentBundle = GetEditorContentBundle(EditingContentBundleGuid))
+		{
+			StartEditing(NewEditingContentBundle);
+		}
+	}
+
+	PreUndoRedoEditingContentBundleGuid.Invalidate();
+}
+
+void UContentBundleEditingSubmodule::OnExecuteActorEditorContextAction(UWorld* InWorld, const EActorEditorContextAction& InType, AActor* InActor /* = nullptr */)
+{
 	switch (InType)
 	{
 	case EActorEditorContextAction::ApplyContext:
-		if (ContentBundleEditorPin != nullptr)
+		if (TSharedPtr<FContentBundleEditor> EditingContentBundle = GetEditorContentBundle(EditingContentBundleGuid))
 		{
-			ContentBundleEditorPin->AddActor(InActor);
+			EditingContentBundle->AddActor(InActor);
 		}
 		break;
 	case EActorEditorContextAction::ResetContext:
+		if (TSharedPtr<FContentBundleEditor> EditingContentBundle = GetEditorContentBundle(EditingContentBundleGuid))
+		{
+			DeactivateCurrentContentBundleEditing();
+		}
 		break;
 	case EActorEditorContextAction::PushContext:
 		PushContentBundleEditing();
@@ -52,9 +80,9 @@ void UContentBundleEditionSubmodule::OnExecuteActorEditorContextAction(UWorld* I
 	}
 }
 
-bool UContentBundleEditionSubmodule::GetActorEditorContextDisplayInfo(UWorld* InWorld, FActorEditorContextClientDisplayInfo& OutDiplayInfo) const
+bool UContentBundleEditingSubmodule::GetActorEditorContextDisplayInfo(UWorld* InWorld, FActorEditorContextClientDisplayInfo& OutDiplayInfo) const
 {
-	if (EditingContentBundle != nullptr)
+	if (EditingContentBundleGuid.IsValid())
 	{
 		constexpr TCHAR PLACEHOLDER_ContentBundle_Brush[] = TEXT("DataLayer.Editor"); // todo_ow: create placeholder brush for content bundle
 		OutDiplayInfo.Title = TEXT("Content Bundle");
@@ -65,18 +93,13 @@ bool UContentBundleEditionSubmodule::GetActorEditorContextDisplayInfo(UWorld* In
 	return false;
 }
 
-bool UContentBundleEditionSubmodule::CanResetContext(UWorld* InWorld) const
-{
-	return false;
-}
-
-TSharedRef<SWidget> UContentBundleEditionSubmodule::GetActorEditorContextWidget(UWorld* World) const
+TSharedRef<SWidget> UContentBundleEditingSubmodule::GetActorEditorContextWidget(UWorld* World) const
 {
 	constexpr TCHAR PLACEHOLDER_ContentBundle_ColorIcon[] = TEXT("DataLayer.ColorIcon"); // todo_ow: create placeholder color icon for content bundle
 
 	TSharedRef<SVerticalBox> OutWidget = SNew(SVerticalBox);
 
-	if (TSharedPtr<FContentBundleEditor> ContentBundleEditorPin = EditingContentBundle.Pin())
+	if (EditingContentBundleGuid.IsValid())
 	{
 		OutWidget->AddSlot().AutoHeight()
 		[
@@ -89,9 +112,9 @@ TSharedRef<SWidget> UContentBundleEditionSubmodule::GetActorEditorContextWidget(
 				SNew(SImage)
 				.ColorAndOpacity_Lambda([this]()
 				{
-					if (TSharedPtr<FContentBundleEditor> ContentBundleEditorPin = EditingContentBundle.Pin())
+					if (TSharedPtr<FContentBundleEditor> ContentBundleEditor = GetEditorContentBundle(EditingContentBundleGuid))
 					{
-						return ContentBundleEditorPin->GetDebugColor();
+						return ContentBundleEditor->GetDebugColor();
 					}
 					return FColor::Black;
 				})
@@ -105,16 +128,15 @@ TSharedRef<SWidget> UContentBundleEditionSubmodule::GetActorEditorContextWidget(
 				SNew(STextBlock)
 				.Text_Lambda([this]()
 				{
-					if (TSharedPtr<FContentBundleEditor> ContentBundleEditorPin = EditingContentBundle.Pin())
+					if (TSharedPtr<FContentBundleEditor> ContentBundleEditor = GetEditorContentBundle(EditingContentBundleGuid))
 					{
-						if (!ContentBundleEditorPin->GetDisplayName().IsEmpty())
+						if (!ContentBundleEditor->GetDisplayName().IsEmpty())
 						{
-							return FText::FromString(ContentBundleEditorPin->GetDisplayName());
+							return FText::FromString(ContentBundleEditor->GetDisplayName());
 						}
 					}
 					return FText::FromString(TEXT("Unset Display Name"));
 				})
-					
 			]
 		];
 	}
@@ -123,96 +145,116 @@ TSharedRef<SWidget> UContentBundleEditionSubmodule::GetActorEditorContextWidget(
 	return OutWidget;
 }
 
-bool UContentBundleEditionSubmodule::ActivateContentBundleEditing(TSharedPtr<FContentBundleEditor>& ContentBundleEditor)
+bool UContentBundleEditingSubmodule::ActivateContentBundleEditing(TSharedPtr<FContentBundleEditor>& ContentBundleEditor)
 {
+	Modify();
+
 	EContentBundleStatus ContentBundleStatus = ContentBundleEditor->GetStatus();
 	bool bCanActivateEditing = ContentBundleStatus == EContentBundleStatus::ContentInjected || ContentBundleStatus == EContentBundleStatus::ReadyToInject;
 	if (bCanActivateEditing && !ContentBundleEditor->IsBeingEdited())
 	{
 		if (IsEditingContentBundle())
 		{
-			check(!IsEditingContentBundle(ContentBundleEditor.Get()));
+			check(!IsEditingContentBundle(ContentBundleEditor));
 			DeactivateCurrentContentBundleEditing();
 		}
-
-		EditingContentBundle = ContentBundleEditor;
-
-		ContentBundleEditor->StartEditing();
-
-		ActorEditorContextClientChanged.Broadcast(this);
-
-		UE_LOG(LogContentBundle, Log, TEXT("[CB: %s] Content Bundle is being edited"), *ContentBundleEditor->GetDescriptor()->GetDisplayName());
-
-		GetSubsystem()->NotifyContentBundleChanged(ContentBundleEditor.Get());
-
+		EditingContentBundleGuid = ContentBundleEditor->GetDescriptor()->GetGuid();
+		StartEditing(ContentBundleEditor);
 		return true;
 	}
 	return false;
 }
 
-bool UContentBundleEditionSubmodule::DeactivateContentBundleEditing(TSharedPtr<FContentBundleEditor>& ContentBundleEditor)
+bool UContentBundleEditingSubmodule::DeactivateCurrentContentBundleEditing()
 {
-	if (IsEditingContentBundle(ContentBundleEditor.Get()))
+	if (TSharedPtr<FContentBundleEditor> ContentBundleToDeactivate = GetEditorContentBundle(EditingContentBundleGuid))
 	{
-		DeactivateCurrentContentBundleEditing();
-		return true;
+		return DeactivateContentBundleEditing(ContentBundleToDeactivate);
 	}
-
 	return false;
 }
 
-bool UContentBundleEditionSubmodule::DeactivateCurrentContentBundleEditing()
+bool UContentBundleEditingSubmodule::DeactivateContentBundleEditing(TSharedPtr<FContentBundleEditor>& ContentBundleEditor)
 {
-	if (TSharedPtr<FContentBundleEditor> ContentBundleEditorPin = EditingContentBundle.Pin())
-	{
-		ContentBundleEditorPin->StopEditing();
-		EditingContentBundle = nullptr;
-		ActorEditorContextClientChanged.Broadcast(this);
-
-		GetSubsystem()->NotifyContentBundleChanged(ContentBundleEditorPin.Get());
-
-		UE_LOG(LogContentBundle, Log, TEXT("[CB: %s] Content Bundle is no longer being edited"), *ContentBundleEditorPin->GetDescriptor()->GetDisplayName());
+	Modify();
 
+	if (IsEditingContentBundle(ContentBundleEditor))
+	{
+		check(EditingContentBundleGuid == ContentBundleEditor->GetDescriptor()->GetGuid());
+		EditingContentBundleGuid.Invalidate();
+		StopEditing(ContentBundleEditor);
 		return true;
 	}
 
 	return false;
 }
 
-void UContentBundleEditionSubmodule::PushContentBundleEditing()
+void UContentBundleEditingSubmodule::StartEditing(TSharedPtr<FContentBundleEditor>& ContentBundleEditor)
 {
-	ActiveContentBundlesStack.Add(EditingContentBundle);
-	DeactivateCurrentContentBundleEditing();
+	check(ContentBundleEditor.IsValid());
+	check(!ContentBundleEditor->IsBeingEdited());
+	ContentBundleEditor->StartEditing();
+	ActorEditorContextClientChanged.Broadcast(this);
+	GetSubsystem()->NotifyContentBundleChanged(ContentBundleEditor.Get());
+	UE_LOG(LogContentBundle, Log, TEXT("[CB: %s] Content Bundle is being edited"), *ContentBundleEditor->GetDescriptor()->GetDisplayName());
+}
+
+void UContentBundleEditingSubmodule::StopEditing(TSharedPtr<FContentBundleEditor>& ContentBundleEditor)
+{
+	check(ContentBundleEditor.IsValid());
+	check(ContentBundleEditor->IsBeingEdited());
+	ContentBundleEditor->StopEditing();
+	ActorEditorContextClientChanged.Broadcast(this);
+	GetSubsystem()->NotifyContentBundleChanged(ContentBundleEditor.Get());
+	UE_LOG(LogContentBundle, Log, TEXT("[CB: %s] Content Bundle is no longer being edited"), *ContentBundleEditor->GetDescriptor()->GetDisplayName());
 }
 
-void UContentBundleEditionSubmodule::PopContentBundleEditing()
+void UContentBundleEditingSubmodule::PushContentBundleEditing()
 {
+	Modify();
+
+	EditingContentBundlesStack.Add(EditingContentBundleGuid);
 	DeactivateCurrentContentBundleEditing();
+}
+
+void UContentBundleEditingSubmodule::PopContentBundleEditing()
+{
+	Modify();
 
-	TSharedPtr<FContentBundleEditor> ContentBundleToActivate = ActiveContentBundlesStack.Pop().Pin();
-	if (ContentBundleToActivate != nullptr)
+	DeactivateCurrentContentBundleEditing();
+	const FGuid ContentBundleGuidToActivate = EditingContentBundlesStack.Pop();
+	if (TSharedPtr<FContentBundleEditor> ContentBundleToActivate = GetEditorContentBundle(ContentBundleGuidToActivate))
 	{
 		ActivateContentBundleEditing(ContentBundleToActivate);
 	}
 }
 
-bool UContentBundleEditionSubmodule::IsEditingContentBundle(const FGuid& ContentBundleGuid) const
+bool UContentBundleEditingSubmodule::IsEditingContentBundle() const
 {
-	TSharedPtr<FContentBundleEditor> ContentBundleEditorPin = EditingContentBundle.Pin();
-	return ContentBundleEditorPin && (ContentBundleEditorPin->GetDescriptor()->GetGuid() == ContentBundleGuid);
+	return IsEditingContentBundle(EditingContentBundleGuid);
 }
 
-UContentBundleEditorSubsystem::UContentBundleEditorSubsystem()
+bool UContentBundleEditingSubmodule::IsEditingContentBundle(const TSharedPtr<FContentBundleEditor>& ContentBundleEditor) const
 {
-	
+	return ContentBundleEditor.IsValid() && IsEditingContentBundle(ContentBundleEditor->GetDescriptor()->GetGuid());
+}
+
+bool UContentBundleEditingSubmodule::IsEditingContentBundle(const FGuid& ContentBundleGuid) const
+{
+	return ContentBundleGuid.IsValid() && (ContentBundleGuid == EditingContentBundleGuid) && GetEditorContentBundle(ContentBundleGuid);
+}
+
+TSharedPtr<FContentBundleEditor> UContentBundleEditingSubmodule::GetEditorContentBundle(const FGuid& ContentBundleGuid) const
+{
+	return GetSubsystem()->GetEditorContentBundle(ContentBundleGuid);
 }
 
 void UContentBundleEditorSubsystem::Initialize(FSubsystemCollectionBase& Collection)
 {
 	Collection.InitializeDependency<UActorEditorContextSubsystem>();
 
-	ContentBundleEditionSubModule = NewObject<UContentBundleEditionSubmodule>(this);
-	ContentBundleEditionSubModule->Initialize();
+	ContentBundleEditingSubModule = NewObject<UContentBundleEditingSubmodule>(this, NAME_None, RF_Transactional);
+	ContentBundleEditingSubModule->Initialize();
 
 	IContentBundleEditorSubsystemInterface::SetInstance(this);
 }
@@ -221,8 +263,8 @@ void UContentBundleEditorSubsystem::Deinitialize()
 {
 	IContentBundleEditorSubsystemInterface::SetInstance(nullptr);
 
-	ContentBundleEditionSubModule->Deinitialize();
-	ContentBundleEditionSubModule = nullptr;
+	ContentBundleEditingSubModule->Deinitialize();
+	ContentBundleEditingSubModule = nullptr;
 }
 
 UWorld* UContentBundleEditorSubsystem::GetWorld() const
@@ -239,7 +281,7 @@ void UContentBundleEditorSubsystem::NotifyContentBundleRemovedContent(const FCon
 {
 	if (ContentBundle->IsBeingEdited())
 	{
-		GetEditionSubmodule()->DeactivateCurrentContentBundleEditing();
+		GetEditingSubmodule()->DeactivateCurrentContentBundleEditing();
 	}
 
 	NotifyContentBundleChanged(ContentBundle);
@@ -279,19 +321,27 @@ TSharedPtr<FContentBundleEditor> UContentBundleEditorSubsystem::GetEditorContent
 	return nullptr;
 }
 
-bool UContentBundleEditorSubsystem::IsEditingContentBundle() const
+TSharedPtr<FContentBundleEditor> UContentBundleEditorSubsystem::GetEditorContentBundle(const FGuid& ContentBundleGuid) const
 {
-	return ContentBundleEditionSubModule->IsEditingContentBundle();
+	if (ContentBundleGuid.IsValid())
+	{
+		if (UContentBundleManager* ContentBundleManager = GetWorld()->ContentBundleManager)
+		{
+			return ContentBundleManager->GetEditorContentBundle(ContentBundleGuid);
+		}
+	}
+
+	return nullptr;
 }
 
-bool UContentBundleEditorSubsystem::IsEditingContentBundle(const FGuid& ContentBundleGuid) const
+bool UContentBundleEditorSubsystem::IsEditingContentBundle() const
 {
-	return ContentBundleEditionSubModule->IsEditingContentBundle(ContentBundleGuid);
+	return ContentBundleEditingSubModule->IsEditingContentBundle();
 }
 
-bool UContentBundleEditorSubsystem::HasContentBundle(const UContentBundleDescriptor* ContentBundleDescriptor) const
+bool UContentBundleEditorSubsystem::IsEditingContentBundle(const FGuid& ContentBundleGuid) const
 {
-	return GetEditorContentBundle(ContentBundleDescriptor) != nullptr;
+	return ContentBundleEditingSubModule->IsEditingContentBundle(ContentBundleGuid);
 }
 
 void UContentBundleEditorSubsystem::SelectActors(FContentBundleEditor& EditorContentBundle)
@@ -316,17 +366,17 @@ void UContentBundleEditorSubsystem::UnreferenceAllActors(FContentBundleEditor& E
 
 bool UContentBundleEditorSubsystem::ActivateContentBundleEditing(TSharedPtr<FContentBundleEditor>& ContentBundleEditor)  const
 {
-	return ContentBundleEditionSubModule->ActivateContentBundleEditing(ContentBundleEditor);
+	return ContentBundleEditingSubModule->ActivateContentBundleEditing(ContentBundleEditor);
 }
 
 bool UContentBundleEditorSubsystem::DeactivateContentBundleEditing(TSharedPtr<FContentBundleEditor>& ContentBundleEditor) const
 {
-	return ContentBundleEditionSubModule->DeactivateContentBundleEditing(ContentBundleEditor);
+	return ContentBundleEditingSubModule->DeactivateContentBundleEditing(ContentBundleEditor);
 }
 
 bool UContentBundleEditorSubsystem::IsContentBundleEditingActivated(TSharedPtr<FContentBundleEditor>& ContentBundleEditor) const
 {
-	return ContentBundleEditionSubModule->IsEditingContentBundle(ContentBundleEditor.Get());
+	return ContentBundleEditor.IsValid() && IsEditingContentBundle(ContentBundleEditor->GetDescriptor()->GetGuid());
 }
 
 void UContentBundleEditorSubsystem::SelectActorsInternal(FContentBundleEditor& EditorContentBundle, bool bSelect)
diff --git a/Engine/Source/Editor/WorldPartitionEditor/Private/WorldPartition/ContentBundle/Outliner/ContentBundleHierarchy.cpp b/Engine/Source/Editor/WorldPartitionEditor/Private/WorldPartition/ContentBundle/Outliner/ContentBundleHierarchy.cpp
--- a/Engine/Source/Editor/WorldPartitionEditor/Private/WorldPartition/ContentBundle/Outliner/ContentBundleHierarchy.cpp
+++ b/Engine/Source/Editor/WorldPartitionEditor/Private/WorldPartition/ContentBundle/Outliner/ContentBundleHierarchy.cpp
@@ -79,7 +79,7 @@ void FContentBundleHiearchy::OnContentBundleChanged(const FContentBundleEditor*
 {
 	FSceneOutlinerHierarchyChangedData EventData;
 	EventData.Type = FSceneOutlinerHierarchyChangedData::FullRefresh;
-	EventData.ItemIDs.Emplace(ContentBundle->GetGuid());
+	EventData.ItemIDs.Emplace(ContentBundle->GetTreeItemID());
 	HierarchyChangedEvent.Broadcast(EventData);
 }
 
diff --git a/Engine/Source/Editor/WorldPartitionEditor/Private/WorldPartition/ContentBundle/Outliner/ContentBundleMode.cpp b/Engine/Source/Editor/WorldPartitionEditor/Private/WorldPartition/ContentBundle/Outliner/ContentBundleMode.cpp
--- a/Engine/Source/Editor/WorldPartitionEditor/Private/WorldPartition/ContentBundle/Outliner/ContentBundleMode.cpp
+++ b/Engine/Source/Editor/WorldPartitionEditor/Private/WorldPartition/ContentBundle/Outliner/ContentBundleMode.cpp
@@ -12,6 +12,7 @@
 #include "SceneOutlinerMenuContext.h"
 #include "Algo/Transform.h"
 #include "Engine/DataAsset.h"
+#include "ScopedTransaction.h"
 #include "ToolMenus.h"
 #include "Editor.h"
 
@@ -163,24 +164,51 @@ void FContentBundleMode::RegisterContextMenu()
 				}
 
 				SSceneOutliner* SceneOutliner = Context->SceneOutliner.Pin().Get();
-				
+
 				{
-					FToolMenuSection& Section = InMenu->AddSection("ContentBundleSection", LOCTEXT("ContentBundle", "Content Bundle"));
+					FToolMenuSection& Section = InMenu->AddSection("ContentBundleActorEditorContext", LOCTEXT("ContentBundleActorEditorContext", "Actor Editor Context"));
 
-					Section.AddMenuEntry("EditContentBundleMenuEntry", LOCTEXT("EditContentBundle", "Edit Content Bundle"), FText(), FSlateIcon(),
+					Section.AddMenuEntry("MakeCurrentContentBundle", LOCTEXT("MakeCurrentContentBundle", "Make Current Content Bundle"), FText(), FSlateIcon(),
 						FUIAction(
-							FExecuteAction::CreateLambda([this, SceneOutliner]
+							FExecuteAction::CreateLambda([this, SceneOutliner]()
+							{
+								TSharedPtr<FContentBundleEditor> SelectedContentBundleEditor = GetSelectedContentBundlePin(SceneOutliner);
+								if (SelectedContentBundleEditor && !SelectedContentBundleEditor->IsBeingEdited())
 								{
-									TWeakPtr<FContentBundleEditor> SelectedContentBundleEditor = GetSelectedContentBundle(SceneOutliner);
 									ToggleContentBundleActivation(SelectedContentBundleEditor);
+								}
+							}),
+							FCanExecuteAction::CreateLambda([this, SceneOutliner]
+							{
+								TSharedPtr<FContentBundleEditor> SelectedContentBundleEditor = GetSelectedContentBundlePin(SceneOutliner);
+								return SelectedContentBundleEditor && !SelectedContentBundleEditor->IsBeingEdited();
+							})
+						));
+
+					Section.AddMenuEntry("ClearCurrentContentBundle", LOCTEXT("ClearCurrentContentBundle", "Clear Current Content Bundle"), FText(), FSlateIcon(),
+						FUIAction(
+							FExecuteAction::CreateLambda([this, SceneOutliner]()
+								{
+									UContentBundleEditorSubsystem* ContentBundleEditorSubsystem = UContentBundleEditorSubsystem::Get();
+									if (ContentBundleEditorSubsystem && ContentBundleEditorSubsystem->IsEditingContentBundle())
+									{
+										if (UContentBundleEditingSubmodule* ContentBundleEditingSubModule = ContentBundleEditorSubsystem->GetEditingSubmodule())
+										{
+											const FScopedTransaction Transaction(LOCTEXT("ClearCurrentContentBundle", "Clear Current Content Bundle"));
+											ContentBundleEditingSubModule->DeactivateCurrentContentBundleEditing();
+										}
+									}
 								}),
-							FCanExecuteAction(),
-							FIsActionChecked::CreateLambda([this, SceneOutliner]
+							FCanExecuteAction::CreateLambda([this, SceneOutliner]
 								{
-									TSharedPtr<FContentBundleEditor> SelectedContentBundleEditor = GetSelectedContentBundlePin(SceneOutliner);
-									return SelectedContentBundleEditor != nullptr && SelectedContentBundleEditor->IsBeingEdited();
+									UContentBundleEditorSubsystem* ContentBundleEditorSubsystem = UContentBundleEditorSubsystem::Get();
+									return ContentBundleEditorSubsystem && ContentBundleEditorSubsystem->IsEditingContentBundle();
 								})
-							));
+						));
+				}
+				
+				{
+					FToolMenuSection& Section = InMenu->AddSection("ContentBundleSection", LOCTEXT("ContentBundle", "Content Bundle"));
 
 					Section.AddMenuEntry("InjectBaseContentMenuEntry", LOCTEXT("InjectBaseContent", "Inject Base Content"), FText(), FSlateIcon(),
 						FUIAction(
@@ -294,17 +322,19 @@ void FContentBundleMode::ToggleContentBundleActivation(const TWeakPtr<FContentBu
 {
 	if (UContentBundleEditorSubsystem* ContentBundleEditorSubsystem = UContentBundleEditorSubsystem::Get())
 	{
-		if (UContentBundleEditionSubmodule* ContentBundleEditionModule = ContentBundleEditorSubsystem->GetEditionSubmodule())
+		if (UContentBundleEditingSubmodule* ContentBundleEditingSubModule = ContentBundleEditorSubsystem->GetEditingSubmodule())
 		{
 			if (TSharedPtr<FContentBundleEditor> ContentBundlePin = ContentBundle.Pin())
 			{
 				if (ContentBundlePin->IsBeingEdited())
 				{
-					ContentBundleEditionModule->DeactivateCurrentContentBundleEditing();
+					const FScopedTransaction Transaction(LOCTEXT("ClearCurrentContentBundle", "Clear Current Content Bundle"));
+					ContentBundleEditingSubModule->DeactivateCurrentContentBundleEditing();
 				}
 				else
 				{
-					ContentBundleEditionModule->ActivateContentBundleEditing(ContentBundlePin);
+					const FScopedTransaction Transaction(LOCTEXT("MakeCurrentContentBundle", "Make Current Content Bundle"));
+					ContentBundleEditingSubModule->ActivateContentBundleEditing(ContentBundlePin);
 				}
 			}
 		}
diff --git a/Engine/Source/Editor/WorldPartitionEditor/Private/WorldPartition/ContentBundle/Outliner/ContentBundleTreeItem.cpp b/Engine/Source/Editor/WorldPartitionEditor/Private/WorldPartition/ContentBundle/Outliner/ContentBundleTreeItem.cpp
--- a/Engine/Source/Editor/WorldPartitionEditor/Private/WorldPartition/ContentBundle/Outliner/ContentBundleTreeItem.cpp
+++ b/Engine/Source/Editor/WorldPartitionEditor/Private/WorldPartition/ContentBundle/Outliner/ContentBundleTreeItem.cpp
@@ -30,7 +30,7 @@ FSceneOutlinerTreeItemID FContentBundleTreeItem::GetID() const
 	TSharedPtr<FContentBundleEditor> ContentBundleEditorPin = ContentBundleEditor.Pin();
 	if (ContentBundleEditorPin != nullptr)
 	{
-		return FSceneOutlinerTreeItemID(ContentBundleEditorPin->GetGuid());
+		return FSceneOutlinerTreeItemID(ContentBundleEditorPin->GetTreeItemID());
 	}
 
 	return FSceneOutlinerTreeItemID();
diff --git a/Engine/Source/Editor/WorldPartitionEditor/Private/WorldPartition/ContentBundle/Outliner/SContentBundleOutliner.cpp b/Engine/Source/Editor/WorldPartitionEditor/Private/WorldPartition/ContentBundle/Outliner/SContentBundleOutliner.cpp
--- a/Engine/Source/Editor/WorldPartitionEditor/Private/WorldPartition/ContentBundle/Outliner/SContentBundleOutliner.cpp
+++ b/Engine/Source/Editor/WorldPartitionEditor/Private/WorldPartition/ContentBundle/Outliner/SContentBundleOutliner.cpp
@@ -9,7 +9,7 @@ void SContentBundleOutliner::SelectContentBundle(const TWeakPtr<FContentBundleEd
 	TSharedPtr<FContentBundleEditor> ContentBundleEditorPin = ContentBundle.Pin();
 	if (ContentBundleEditorPin != nullptr)
 	{
-		FSceneOutlinerTreeItemPtr Item = GetTreeItem(ContentBundleEditorPin->GetGuid());
+		FSceneOutlinerTreeItemPtr Item = GetTreeItem(ContentBundleEditorPin->GetTreeItemID());
 		if (Item->IsValid())
 		{
 			SetItemSelection(Item, true, ESelectInfo::Direct);
diff --git a/Engine/Source/Editor/WorldPartitionEditor/Public/WorldPartition/ContentBundle/ContentBundleEditorSubsystem.h b/Engine/Source/Editor/WorldPartitionEditor/Public/WorldPartition/ContentBundle/ContentBundleEditorSubsystem.h
--- a/Engine/Source/Editor/WorldPartitionEditor/Public/WorldPartition/ContentBundle/ContentBundleEditorSubsystem.h
+++ b/Engine/Source/Editor/WorldPartitionEditor/Public/WorldPartition/ContentBundle/ContentBundleEditorSubsystem.h
@@ -40,7 +40,7 @@ private:
 };
 
 UCLASS()
-class UContentBundleEditionSubmodule : public UContentBundleEditorSubsystemModule, public IActorEditorContextClient
+class UContentBundleEditingSubmodule : public UContentBundleEditorSubsystemModule, public IActorEditorContextClient
 {
 	GENERATED_BODY()
 
@@ -48,7 +48,7 @@ public:
 	//~ Begin IActorEditorContextClient interface
 	virtual void OnExecuteActorEditorContextAction(UWorld* InWorld, const EActorEditorContextAction& InType, AActor* InActor = nullptr) override;
 	virtual bool GetActorEditorContextDisplayInfo(UWorld* InWorld, FActorEditorContextClientDisplayInfo& OutDiplayInfo) const override;
-	virtual bool CanResetContext(UWorld* InWorld) const override;
+	virtual bool CanResetContext(UWorld* InWorld) const override { return true; }
 	virtual TSharedRef<SWidget> GetActorEditorContextWidget(UWorld* InWorld) const override;
 	virtual FOnActorEditorContextClientChanged& GetOnActorEditorContextClientChanged() override { return ActorEditorContextClientChanged; }
 	//~ End IActorEditorContextClient interface
@@ -56,9 +56,10 @@ public:
 	bool ActivateContentBundleEditing(TSharedPtr<FContentBundleEditor>& ContentBundleEditor);
 	bool DeactivateContentBundleEditing(TSharedPtr<FContentBundleEditor>& ContentBundleEditor);
 	bool DeactivateCurrentContentBundleEditing();
-	bool IsEditingContentBundle() const { return EditingContentBundle.IsValid(); }
-	bool IsEditingContentBundle(const FContentBundleEditor* ContentBundleEditor) const { return EditingContentBundle.Pin().Get() == ContentBundleEditor; }
+	bool IsEditingContentBundle() const;
+	bool IsEditingContentBundle(const TSharedPtr<FContentBundleEditor>& ContentBundleEditor) const;
 	bool IsEditingContentBundle(const FGuid& ContentBundleGuid) const;
+	TSharedPtr<FContentBundleEditor> GetEditorContentBundle(const FGuid& ContentBundleGuid) const;
 
 protected:
 	//~ Begin UContentBundleEditorSubsystemModule interface
@@ -69,10 +70,19 @@ protected:
 private:
 	void PushContentBundleEditing();
 	void PopContentBundleEditing();
+	void StartEditing(TSharedPtr<FContentBundleEditor>& ContentBundleEditor);
+	void StopEditing(TSharedPtr<FContentBundleEditor>& ContentBundleEditor);
+	virtual void PreEditUndo() override;
+	virtual void PostEditUndo() override;
 
-	TArray<TWeakPtr<FContentBundleEditor>> ActiveContentBundlesStack;
+	UPROPERTY()
+	TArray<FGuid> EditingContentBundlesStack;
 
-	TWeakPtr<FContentBundleEditor> EditingContentBundle;
+	UPROPERTY()
+	FGuid EditingContentBundleGuid;
+
+	// Used for undo/redo
+	FGuid PreUndoRedoEditingContentBundleGuid;
 
 	FOnActorEditorContextClientChanged ActorEditorContextClientChanged;
 };
@@ -85,7 +95,7 @@ class WORLDPARTITIONEDITOR_API UContentBundleEditorSubsystem : public UEditorSub
 public:
 	static UContentBundleEditorSubsystem* Get() { return StaticCast<UContentBundleEditorSubsystem*>(IContentBundleEditorSubsystemInterface::Get()); }
 
-	UContentBundleEditorSubsystem();
+	UContentBundleEditorSubsystem() {}
 
 	//~ Begin UEditorSubsystem interface
 	virtual void Initialize(FSubsystemCollectionBase& Collection) override;
@@ -102,14 +112,13 @@ public:
 
 	UWorld* GetWorld() const;
 
-	UContentBundleEditionSubmodule* GetEditionSubmodule() { return ContentBundleEditionSubModule; }
+	UContentBundleEditingSubmodule* GetEditingSubmodule() { return ContentBundleEditingSubModule; }
 
 	TSharedPtr<FContentBundleEditor> GetEditorContentBundleForActor(const AActor* Actor);
 
 	TArray<TSharedPtr<FContentBundleEditor>> GetEditorContentBundles();
 	TSharedPtr<FContentBundleEditor> GetEditorContentBundle(const UContentBundleDescriptor* ContentBundleDescriptor) const;
-
-	bool HasContentBundle(const UContentBundleDescriptor* ContentBundleDescriptor) const;
+	TSharedPtr<FContentBundleEditor> GetEditorContentBundle(const FGuid& ContentBundleGuid) const;
 
 	void SelectActors(FContentBundleEditor& EditorContentBundle);
 	void DeselectActors(FContentBundleEditor& EditorContentBundle);
@@ -136,7 +145,7 @@ private:
 	void SelectActorsInternal(FContentBundleEditor& EditorContentBundle, bool bSelect);
 
 	UPROPERTY()
-	TObjectPtr<UContentBundleEditionSubmodule> ContentBundleEditionSubModule;
+	TObjectPtr<UContentBundleEditingSubmodule> ContentBundleEditingSubModule;
 
 	FOnContentBundleChanged ContentBundleChanged;
 	FOnContentBundleAdded ContentBundleAdded;
diff --git a/Engine/Source/Runtime/Engine/Private/WorldPartition/ContentBundle/ContentBundleEditor.cpp b/Engine/Source/Runtime/Engine/Private/WorldPartition/ContentBundle/ContentBundleEditor.cpp
--- a/Engine/Source/Runtime/Engine/Private/WorldPartition/ContentBundle/ContentBundleEditor.cpp
+++ b/Engine/Source/Runtime/Engine/Private/WorldPartition/ContentBundle/ContentBundleEditor.cpp
@@ -26,7 +26,7 @@ FContentBundleEditor::FContentBundleEditor(TSharedPtr<FContentBundleClient>& InC
 	: FContentBundleBase(InClient, InWorld)
 	, UnsavedActorMonitor(nullptr)
 	, ExternalStreamingObject(nullptr)
-	, Guid(FGuid::NewGuid())
+	, TreeItemID(FGuid::NewGuid())
 	, bIsBeingEdited(false)
 {}
 
diff --git a/Engine/Source/Runtime/Engine/Public/WorldPartition/ContentBundle/ContentBundleEditor.h b/Engine/Source/Runtime/Engine/Public/WorldPartition/ContentBundle/ContentBundleEditor.h
--- a/Engine/Source/Runtime/Engine/Public/WorldPartition/ContentBundle/ContentBundleEditor.h
+++ b/Engine/Source/Runtime/Engine/Public/WorldPartition/ContentBundle/ContentBundleEditor.h
@@ -34,7 +34,7 @@ public:
 	virtual bool IsValid() const override;
 	//~ End IContentBundle Interface
 
-	const FGuid& GetGuid() const { return Guid; }
+	const FGuid& GetTreeItemID() const { return TreeItemID; }
 
 	void StartEditing();
 	void StopEditing();
@@ -94,7 +94,8 @@ private:
 	TMap<uint32, const UWorldPartitionRuntimeCell*> CookPackageIdsToCell;
 	URuntimeHashExternalStreamingObjectBase* ExternalStreamingObject;
 
-	FGuid Guid;
+	// Used by Content Bundle Outliner to link tree item to FContentBundle instance.
+	FGuid TreeItemID;
 
 	bool bIsBeingEdited;
 };
