Re-load current map when Build Minimap or Build HLODs commandlet fails

diff --git a/Engine/Source/Editor/WorldPartitionEditor/Private/WorldPartitionEditorModule.cpp b/Engine/Source/Editor/WorldPartitionEditor/Private/WorldPartitionEditorModule.cpp
--- a/Engine/Source/Editor/WorldPartitionEditor/Private/WorldPartitionEditorModule.cpp
+++ b/Engine/Source/Editor/WorldPartitionEditor/Private/WorldPartitionEditorModule.cpp
@@ -599,6 +599,9 @@ bool FWorldPartitionEditorModule::Build(const FRunBuilderParams& InParams)
 		return false;
 	}
 
+	// Close assets editors as the external UE process may try to update those same assets
+	GEditor->GetEditorSubsystem<UAssetEditorSubsystem>()->CloseAllAssetEditors();
+
 	TStringBuilder<512> CommandletArgsBuilder;
 	CommandletArgsBuilder.Append(MapPackage);
 	CommandletArgsBuilder.Append(" -run=WorldPartitionBuilderCommandlet -Builder=");
@@ -619,12 +622,10 @@ bool FWorldPartitionEditorModule::Build(const FRunBuilderParams& InParams)
 	FString CommandletArgs(CommandletArgsBuilder.ToString());
 	RunCommandletAsExternalProcess(CommandletArgs, OperationDescription, Result, bCancelled, CommandletOutput);
 
-	bool bSuccess = !bCancelled && Result == 0;
-	if (bSuccess)
-	{
-		RescanAssetsAndLoadMap(MapPackage);
-	}
-	else if (bCancelled)
+	// Reload map
+	RescanAssetsAndLoadMap(MapPackage);
+
+	if (bCancelled)
 	{
 		FMessageDialog::Open(EAppMsgType::Ok, LOCTEXT("BuildCancelled", "Build cancelled!"));
 	}
@@ -633,6 +634,7 @@ bool FWorldPartitionEditorModule::Build(const FRunBuilderParams& InParams)
 		FMessageDialog::Open(EAppMsgType::Ok, LOCTEXT("BuildFailed", "Build failed! See log for details."));
 	}
 
+	const bool bSuccess = !bCancelled && Result == 0;
 	return bSuccess;
 }
 
