In editor World Partition builds that launch external commandlets are now cancellable

diff --git a/Engine/Source/Editor/WorldPartitionEditor/Private/WorldPartitionEditorModule.cpp b/Engine/Source/Editor/WorldPartitionEditor/Private/WorldPartitionEditorModule.cpp
--- a/Engine/Source/Editor/WorldPartitionEditor/Private/WorldPartitionEditorModule.cpp
+++ b/Engine/Source/Editor/WorldPartitionEditor/Private/WorldPartitionEditorModule.cpp
@@ -394,7 +394,7 @@ void FWorldPartitionEditorModule::RunCommandletAsExternalProcess(const FString&
 
 	while (FPlatformProcess::IsProcRunning(ProcessHandle))
 	{
-		if (SlowTask.ShouldCancel())
+		if (SlowTask.ShouldCancel() || GEditor->GetMapBuildCancelled())
 		{
 			bOutCancelled = true;
 			FPlatformProcess::TerminateProc(ProcessHandle);
