World Partition Editor Build Commandlets:
- Added pre and post run events for the builder commandlets ran from the editor.
- Added the possibility to avoid unloading and reloading the world when running a builder commandlet.

diff --git a/Engine/Source/Editor/SceneOutliner/Private/ActorDescTreeItem.cpp b/Engine/Source/Editor/SceneOutliner/Private/ActorDescTreeItem.cpp
--- a/Engine/Source/Editor/SceneOutliner/Private/ActorDescTreeItem.cpp
+++ b/Engine/Source/Editor/SceneOutliner/Private/ActorDescTreeItem.cpp
@@ -353,7 +353,7 @@ bool FActorDescTreeItem::ShouldShowPinnedState() const
 
 bool FActorDescTreeItem::GetPinnedState() const
 {
-	if (ActorDescHandle.IsValid())
+	if (ActorDescHandle.IsValid() && ActorDescHandle->GetContainer())
 	{
 		UWorldPartition* WorldPartition = ActorDescHandle->GetContainer()->GetWorld()->GetWorldPartition();
 		return WorldPartition ? WorldPartition->IsActorPinned(GetGuid()) : false;
diff --git a/Engine/Source/Editor/WorldPartitionEditor/Private/WorldPartitionEditorModule.cpp b/Engine/Source/Editor/WorldPartitionEditor/Private/WorldPartitionEditorModule.cpp
--- a/Engine/Source/Editor/WorldPartitionEditor/Private/WorldPartitionEditorModule.cpp
+++ b/Engine/Source/Editor/WorldPartitionEditor/Private/WorldPartitionEditorModule.cpp
@@ -326,21 +326,24 @@ static bool UnloadCurrentMap(bool bAskSaveContentPackages, FString& MapPackageNa
 	return true;
 }
 
-static void RescanAssetsAndLoadMap(const FString& MapToLoad)
+static void RescanAssets(const FString& MapToScan)
 {
 	// Force a directory watcher tick for the asset registry to get notified of the changes
 	FDirectoryWatcherModule& DirectoryWatcherModule = FModuleManager::Get().LoadModuleChecked<FDirectoryWatcherModule>(TEXT("DirectoryWatcher"));
 	DirectoryWatcherModule.Get()->Tick(-1.0f);
 
-	// Force update before loading converted map
+	// Force update
 	FAssetRegistryModule& AssetRegistryModule = FModuleManager::LoadModuleChecked<FAssetRegistryModule>("AssetRegistry");
 	IAssetRegistry& AssetRegistry = AssetRegistryModule.Get();
 
-	TArray<FString> ExternalObjectsPaths = ULevel::GetExternalObjectsPaths(MapToLoad);
+	TArray<FString> ExternalObjectsPaths = ULevel::GetExternalObjectsPaths(MapToScan);
 
-	AssetRegistry.ScanModifiedAssetFiles({ MapToLoad });
+	AssetRegistry.ScanModifiedAssetFiles({ MapToScan });
 	AssetRegistry.ScanPathsSynchronous(ExternalObjectsPaths, true);
+}
 
+static void LoadMap(const FString& MapToLoad)
+{
 	FEditorFileUtils::LoadMap(MapToLoad);
 
 	UWorld* World = GEditor->GetEditorWorldContext().World();
@@ -505,7 +508,8 @@ bool FWorldPartitionEditorModule::ConvertMap(const FString& InLongPackageName)
 				MapToLoad += UWorldPartitionConvertCommandlet::GetConversionSuffix(DefaultConvertOptions->bOnlyMergeSubLevels);
 			}
 
-			RescanAssetsAndLoadMap(MapToLoad);
+			RescanAssets(MapToLoad);
+			LoadMap(MapToLoad);
 		}
 		else if (bCancelled)
 		{
@@ -593,8 +597,11 @@ bool FWorldPartitionEditorModule::BuildMinimap(const FRunBuilderParams& InParams
 
 bool FWorldPartitionEditorModule::Build(const FRunBuilderParams& InParams)
 {
-	FString MapPackage = InParams.World->GetPackage()->GetName();
-	if (!UnloadCurrentMap(/*bAskSaveContentPackages=*/true, MapPackage))
+	FRunBuilderParams ParamsCopy(InParams);
+	OnPreExecuteCommandletEvent.Broadcast(ParamsCopy);
+
+	FString MapPackage = ParamsCopy.World->GetPackage()->GetName();
+	if (ParamsCopy.bUnloadMap && !UnloadCurrentMap(/*bAskSaveContentPackages=*/true, MapPackage))
 	{
 		return false;
 	}
@@ -605,15 +612,15 @@ bool FWorldPartitionEditorModule::Build(const FRunBuilderParams& InParams)
 	TStringBuilder<512> CommandletArgsBuilder;
 	CommandletArgsBuilder.Append(MapPackage);
 	CommandletArgsBuilder.Append(" -run=WorldPartitionBuilderCommandlet -Builder=");
-	CommandletArgsBuilder.Append(InParams.BuilderClass->GetName());
+	CommandletArgsBuilder.Append(ParamsCopy.BuilderClass->GetName());
 		
-	if (!InParams.ExtraArgs.IsEmpty())
+	if (!ParamsCopy.ExtraArgs.IsEmpty())
 	{
 		CommandletArgsBuilder.Append(" ");
-		CommandletArgsBuilder.Append(InParams.ExtraArgs);
+		CommandletArgsBuilder.Append(ParamsCopy.ExtraArgs);
 	}
 
-	const FText OperationDescription = InParams.OperationDescription.IsEmptyOrWhitespace() ? LOCTEXT("BuildProgress", "Building...") : InParams.OperationDescription;
+	const FText OperationDescription = ParamsCopy.OperationDescription.IsEmptyOrWhitespace() ? LOCTEXT("BuildProgress", "Building...") : ParamsCopy.OperationDescription;
 
 	int32 Result;
 	bool bCancelled;
@@ -622,8 +629,19 @@ bool FWorldPartitionEditorModule::Build(const FRunBuilderParams& InParams)
 	FString CommandletArgs(CommandletArgsBuilder.ToString());
 	RunCommandletAsExternalProcess(CommandletArgs, OperationDescription, Result, bCancelled, CommandletOutput);
 
-	// Reload map
-	RescanAssetsAndLoadMap(MapPackage);
+	RescanAssets(MapPackage);
+
+	if (ParamsCopy.bUnloadMap)
+	{
+		LoadMap(MapPackage);
+	}
+	else
+	{
+		if (UWorldPartition* WorldPartition = ParamsCopy.World->GetWorldPartition())
+		{
+			WorldPartition->Update();
+		}
+	}
 
 	if (bCancelled)
 	{
@@ -634,8 +652,9 @@ bool FWorldPartitionEditorModule::Build(const FRunBuilderParams& InParams)
 		FMessageDialog::Open(EAppMsgType::Ok, LOCTEXT("BuildFailed", "Build failed! See log for details."));
 	}
 
-	const bool bSuccess = !bCancelled && Result == 0;
-	return bSuccess;
+	OnPostExecuteCommandletEvent.Broadcast();
+
+	return !bCancelled && Result == 0;
 }
 
 bool FWorldPartitionEditorModule::BuildLandscapeSplineMeshes(UWorld* InWorld)
diff --git a/Engine/Source/Editor/WorldPartitionEditor/Public/WorldPartitionEditorModule.h b/Engine/Source/Editor/WorldPartitionEditor/Public/WorldPartitionEditorModule.h
--- a/Engine/Source/Editor/WorldPartitionEditor/Public/WorldPartitionEditorModule.h
+++ b/Engine/Source/Editor/WorldPartitionEditor/Public/WorldPartitionEditorModule.h
@@ -78,9 +78,16 @@ public:
 	/** Return the world added event. */
 	virtual FWorldPartitionCreated& OnWorldPartitionCreated() override { return WorldPartitionCreatedEvent; }
 
+	/** Return the commandlet pre-execution event */
+	virtual FOnPreExecuteCommandlet& OnPreExecuteCommandlet() override { return OnPreExecuteCommandletEvent; }
+
 	/** Return the commandlet execution event */
 	virtual FOnExecuteCommandlet& OnExecuteCommandlet() override { return OnExecuteCommandletEvent; }
 
+	/** Return the commandlet post-execution event */
+	virtual FOnPostExecuteCommandlet& OnPostExecuteCommandlet() override { return OnPostExecuteCommandletEvent; }
+
+
 	/**
 	 * Creates a Content Bundle Browser widget
 	 */
@@ -133,5 +140,7 @@ private:
 
 	FWorldPartitionCreated WorldPartitionCreatedEvent;
 
+	FOnPreExecuteCommandlet OnPreExecuteCommandletEvent;
 	FOnExecuteCommandlet OnExecuteCommandletEvent;
+	FOnPostExecuteCommandlet OnPostExecuteCommandletEvent;
 };
diff --git a/Engine/Source/Runtime/Engine/Private/WorldPartition/ActorDescContainer.cpp b/Engine/Source/Runtime/Engine/Private/WorldPartition/ActorDescContainer.cpp
--- a/Engine/Source/Runtime/Engine/Private/WorldPartition/ActorDescContainer.cpp
+++ b/Engine/Source/Runtime/Engine/Private/WorldPartition/ActorDescContainer.cpp
@@ -9,7 +9,6 @@
 #include "UObject/LinkerInstancingContext.h"
 #include "AssetRegistry/AssetRegistryModule.h"
 #include "WorldPartition/WorldPartition.h"
-#include "WorldPartition/WorldPartitionLog.h"
 #include "WorldPartition/WorldPartitionHandle.h"
 #include "WorldPartition/WorldPartitionActorDescUtils.h"
 #include "UObject/ObjectSaveContext.h"
@@ -91,6 +90,73 @@ void UActorDescContainer::Initialize(const FInitializeParams& InitParams)
 #endif
 }
 
+void UActorDescContainer::Update()
+{
+	check(World);
+
+#if WITH_EDITOR
+	check(bContainerInitialized);
+	TArray<FAssetData> Assets;
+
+	const FString ContainerExternalActorsPath = GetExternalActorPath();
+
+	// Do a synchronous scan of the level external actors path.			
+	IAssetRegistry& AssetRegistry = FModuleManager::LoadModuleChecked<FAssetRegistryModule>(TEXT("AssetRegistry")).Get();
+	AssetRegistry.ScanPathsSynchronous({ ContainerExternalActorsPath }, /*bForceRescan*/false, /*bIgnoreDenyListScanFilters*/false);
+
+	FARFilter Filter;
+	Filter.bRecursivePaths = true;
+	Filter.bIncludeOnlyOnDiskAssets = true;
+	Filter.PackagePaths.Add(*ContainerExternalActorsPath);
+
+	AssetRegistry.GetAssets(Filter, Assets);
+
+	TSet<FGuid> ActorGuids;
+	for (const FAssetData& Asset : Assets)
+	{
+		TUniquePtr<FWorldPartitionActorDesc> NewActorDesc = FWorldPartitionActorDescUtils::GetActorDescriptorFromAssetData(Asset);
+
+		if (NewActorDesc.IsValid() && NewActorDesc->GetNativeClass().IsValid())
+		{
+			ActorGuids.Add(NewActorDesc->GetGuid());
+
+			if (TUniquePtr<FWorldPartitionActorDesc>* ExistingActorDesc = GetActorDescriptor(NewActorDesc->GetGuid()))
+			{
+				if (!NewActorDesc->Equals((*ExistingActorDesc).Get()))
+				{
+					OnActorDescUpdating(ExistingActorDesc->Get());
+					FWorldPartitionActorDescUtils::UpdateActorDescriptorFromActorDescriptor(NewActorDesc, *ExistingActorDesc);
+					OnActorDescUpdated(ExistingActorDesc->Get());
+				}
+			}
+			else
+			{
+				FWorldPartitionActorDesc* ActorDescPtr = NewActorDesc.Release();
+				AddActorDescriptor(ActorDescPtr);
+				OnActorDescAdded(ActorDescPtr);
+			}
+		}
+	}
+
+	TArray<FGuid> ActorDescsToRemove;
+	for (FActorDescList::TIterator<> ActorDescIt(this); ActorDescIt; ++ActorDescIt)
+	{
+		if (!ActorGuids.Contains(ActorDescIt->GetGuid()))
+		{
+			ActorDescsToRemove.Add(ActorDescIt->GetGuid());
+		}
+	}
+
+	for (const FGuid& ActorDescGuidToRemove : ActorDescsToRemove)
+	{
+		if (FWorldPartitionActorDesc* ActorDesc = GetActorDesc(ActorDescGuidToRemove))
+		{
+			RemoveActor(ActorDescGuidToRemove);
+		}
+	}
+#endif
+}
+
 void UActorDescContainer::Uninitialize()
 {
 #if WITH_EDITOR
@@ -185,7 +251,7 @@ void UActorDescContainer::OnObjectPreSave(UObject* Object, FObjectPreSaveContext
 				{
 					// Existing actor
 					OnActorDescUpdating(ExistingActorDesc->Get());
-					FWorldPartitionActorDescUtils::UpdateActorDescriptorFomActor(Actor, *ExistingActorDesc);
+					FWorldPartitionActorDescUtils::UpdateActorDescriptorFromActor(Actor, *ExistingActorDesc);
 					OnActorDescUpdated(ExistingActorDesc->Get());
 				}
 				else
diff --git a/Engine/Source/Runtime/Engine/Private/WorldPartition/WorldPartition.cpp b/Engine/Source/Runtime/Engine/Private/WorldPartition/WorldPartition.cpp
--- a/Engine/Source/Runtime/Engine/Private/WorldPartition/WorldPartition.cpp
+++ b/Engine/Source/Runtime/Engine/Private/WorldPartition/WorldPartition.cpp
@@ -541,7 +541,7 @@ void UWorldPartition::Initialize(UWorld* InWorld, const FTransform& InTransform)
 			InstancingContext.AddPathMapping(
 				FSoftObjectPath(*FString::Format(TEXT("{0}.{1}"), {PackageName.ToString(), FPackageName::GetShortName(PackageName)})),
 				FSoftObjectPath(GetWorld())
-				);
+			);
 		}
 
 		ActorDescContainer = RegisterActorDescContainer(PackageName);
@@ -562,6 +562,11 @@ void UWorldPartition::Initialize(UWorld* InWorld, const FTransform& InTransform)
 
 				ActorDescIterator->bIsForcedNonSpatiallyLoaded = !IsStreamingEnabled();
 
+				if (ForceLoadedActors)
+				{
+					ForceLoadedActors->AddActors({ ActorDescIterator->GetGuid() });
+				}
+
 				if (bIsEditor && !bIsCooking)
 				{
 					HashActorDesc(*ActorDescIterator);
@@ -757,6 +762,20 @@ bool UWorldPartition::IsInitialized() const
 	return InitState == EWorldPartitionInitState::Initialized;
 }
 
+void UWorldPartition::Update()
+{
+#if WITH_EDITOR
+	UWorld* OuterWorld = GetTypedOuter<UWorld>();
+	check(OuterWorld);
+	check(!OuterWorld->IsInstanced());
+
+	ForEachActorDescContainer([](UActorDescContainer* InActorDescContainer)
+	{
+		InActorDescContainer->Update();
+	});
+#endif
+}
+
 bool UWorldPartition::SupportsStreaming() const
 {
 	return World ? World->GetWorldSettings()->SupportsWorldPartitionStreaming() : false;
@@ -1071,6 +1090,11 @@ void UWorldPartition::OnActorDescAdded(FWorldPartitionActorDesc* NewActorDesc)
 
 	HashActorDesc(NewActorDesc);
 
+	if (ForceLoadedActors)
+	{
+		ForceLoadedActors->AddActors({ NewActorDesc->GetGuid() });
+	}
+
 	if (WorldPartitionEditor)
 	{
 		WorldPartitionEditor->Refresh();
@@ -1080,6 +1104,11 @@ void UWorldPartition::OnActorDescAdded(FWorldPartitionActorDesc* NewActorDesc)
 void UWorldPartition::OnActorDescRemoved(FWorldPartitionActorDesc* ActorDesc)
 {	
 	UnhashActorDesc(ActorDesc);
+
+	if (ForceLoadedActors)
+	{
+		ForceLoadedActors->RemoveActors({ ActorDesc->GetGuid() });
+	}
 		
 	if (WorldPartitionEditor)
 	{
@@ -1186,11 +1215,6 @@ void UWorldPartition::HashActorDesc(FWorldPartitionActorDesc* ActorDesc)
 	EditorHash->HashActor(ActorHandle);
 
 	bShouldCheckEnableStreamingWarning = IsMainWorldPartition();
-
-	if (ForceLoadedActors)
-	{
-		ForceLoadedActors->AddActors({ ActorDesc->GetGuid() });
-	}
 }
 
 void UWorldPartition::UnhashActorDesc(FWorldPartitionActorDesc* ActorDesc)
@@ -1200,11 +1224,6 @@ void UWorldPartition::UnhashActorDesc(FWorldPartitionActorDesc* ActorDesc)
 
 	FWorldPartitionHandle ActorHandle(this, ActorDesc->GetGuid());
 	EditorHash->UnhashActor(ActorHandle);
-
-	if (ForceLoadedActors)
-	{
-		ForceLoadedActors->RemoveActors({ ActorDesc->GetGuid() });
-	}
 }
 #endif
 
@@ -1744,6 +1763,11 @@ bool UWorldPartition::UnregisterActorDescContainer(UActorDescContainer* InActorD
 
 		UnpinActors(ActorGuids);
 
+		if (ForceLoadedActors)
+		{
+			ForceLoadedActors->RemoveActors(ActorGuids);
+		}
+
 		OnActorDescContainerUnregistered.Broadcast(InActorDescContainer);
 
 		if (IsInitialized() && EditorHash != nullptr)
diff --git a/Engine/Source/Runtime/Engine/Private/WorldPartition/WorldPartitionActorDescUtils.cpp b/Engine/Source/Runtime/Engine/Private/WorldPartition/WorldPartitionActorDescUtils.cpp
--- a/Engine/Source/Runtime/Engine/Private/WorldPartition/WorldPartitionActorDescUtils.cpp
+++ b/Engine/Source/Runtime/Engine/Private/WorldPartition/WorldPartitionActorDescUtils.cpp
@@ -130,11 +130,16 @@ void FWorldPartitionActorDescUtils::AppendAssetDataTagsFromActor(const AActor* I
 	OutTags.Add(UObject::FAssetRegistryTag(NAME_ActorMetaData, ActorMetaData, UObject::FAssetRegistryTag::TT_Hidden));
 }
 
-void FWorldPartitionActorDescUtils::UpdateActorDescriptorFomActor(const AActor* InActor, TUniquePtr<FWorldPartitionActorDesc>& ActorDesc)
+void FWorldPartitionActorDescUtils::UpdateActorDescriptorFromActor(const AActor* InActor, TUniquePtr<FWorldPartitionActorDesc>& OutActorDesc)
 {
 	TUniquePtr<FWorldPartitionActorDesc> NewActorDesc(InActor->CreateActorDesc());
-	NewActorDesc->TransferFrom(ActorDesc.Get());
-	ActorDesc = MoveTemp(NewActorDesc);
+	UpdateActorDescriptorFromActorDescriptor(NewActorDesc, OutActorDesc);
+}
+
+void FWorldPartitionActorDescUtils::UpdateActorDescriptorFromActorDescriptor(TUniquePtr<FWorldPartitionActorDesc>& InActorDesc, TUniquePtr<FWorldPartitionActorDesc>& OutActorDesc)
+{
+	InActorDesc->TransferFrom(OutActorDesc.Get());
+	OutActorDesc = MoveTemp(InActorDesc);
 }
 
 void FWorldPartitionActorDescUtils::ReplaceActorDescriptorPointerFromActor(const AActor* InOldActor, AActor* InNewActor, FWorldPartitionActorDesc* InActorDesc)
diff --git a/Engine/Source/Runtime/Engine/Public/WorldPartition/ActorDescContainer.h b/Engine/Source/Runtime/Engine/Public/WorldPartition/ActorDescContainer.h
--- a/Engine/Source/Runtime/Engine/Public/WorldPartition/ActorDescContainer.h
+++ b/Engine/Source/Runtime/Engine/Public/WorldPartition/ActorDescContainer.h
@@ -44,6 +44,7 @@ public:
 	UE_DEPRECATED(5.1, "UActorDescContainer::Initialize is deprecated, UActorDescContainer::Initialize with UActorDescContainer::FInitializeParams should be used instead.")
 	void Initialize(UWorld* World, FName InPackageName);
 	void Initialize(const FInitializeParams& InitParams);
+	void Update();
 	void Uninitialize();
 	
 	virtual UWorld* GetWorld() const override;
diff --git a/Engine/Source/Runtime/Engine/Public/WorldPartition/IWorldPartitionEditorModule.h b/Engine/Source/Runtime/Engine/Public/WorldPartition/IWorldPartitionEditorModule.h
--- a/Engine/Source/Runtime/Engine/Public/WorldPartition/IWorldPartitionEditorModule.h
+++ b/Engine/Source/Runtime/Engine/Public/WorldPartition/IWorldPartitionEditorModule.h
@@ -22,10 +22,11 @@ public:
 		
 	struct FRunBuilderParams
 	{
-		TSubclassOf<UWorldPartitionBuilder> BuilderClass;
-		UWorld* World;
+		TSubclassOf<UWorldPartitionBuilder> BuilderClass = nullptr;
+		UWorld* World = nullptr;
 		FString ExtraArgs;
 		FText OperationDescription;
+		bool bUnloadMap = false;
 	};
 
 	virtual bool RunBuilder(TSubclassOf<UWorldPartitionBuilder> BuilderClass, UWorld* InWorld);
@@ -50,12 +51,24 @@ public:
 	/** Triggered when a world is added. */
 	DECLARE_EVENT_OneParam(IWorldPartitionEditorModule, FWorldPartitionCreated, UWorld*);
 
+	/** Triggered when the editor is about to launch a commandlet. Can be used to modify the builder params. */
+	DECLARE_EVENT_OneParam(IWorldPartitionEditorModule, FOnPreExecuteCommandlet, FRunBuilderParams&);
+
 	/** Triggered when the editor launches a commandlet. Can be used to provide project specific arguments. */
 	DECLARE_EVENT_OneParam(IWorldPartitionEditorModule, FOnExecuteCommandlet, TArray<FString>&);
 
+	/** Triggered when the editor has launched a commandlet. */
+	DECLARE_EVENT(IWorldPartitionEditorModule, FOnPostExecuteCommandlet);
+
 	/** Return the world added event. */
 	virtual FWorldPartitionCreated& OnWorldPartitionCreated() = 0;
 
+	/** Return the commandlet pre-execution event */
+	virtual FOnPreExecuteCommandlet& OnPreExecuteCommandlet() = 0;
+
 	/** Return the commandlet execution event */
 	virtual FOnExecuteCommandlet& OnExecuteCommandlet() = 0;
+
+	/** Return the commandlet post-execution event */
+	virtual FOnPostExecuteCommandlet& OnPostExecuteCommandlet() = 0;
 };
\ No newline at end of file
diff --git a/Engine/Source/Runtime/Engine/Public/WorldPartition/WorldPartition.h b/Engine/Source/Runtime/Engine/Public/WorldPartition/WorldPartition.h
--- a/Engine/Source/Runtime/Engine/Public/WorldPartition/WorldPartition.h
+++ b/Engine/Source/Runtime/Engine/Public/WorldPartition/WorldPartition.h
@@ -240,6 +240,7 @@ public:
 
 	void Initialize(UWorld* World, const FTransform& InTransform);
 	bool IsInitialized() const;
+	void Update();
 	void Uninitialize();
 
 	bool SupportsStreaming() const;
diff --git a/Engine/Source/Runtime/Engine/Public/WorldPartition/WorldPartitionActorDescUtils.h b/Engine/Source/Runtime/Engine/Public/WorldPartition/WorldPartitionActorDescUtils.h
--- a/Engine/Source/Runtime/Engine/Public/WorldPartition/WorldPartitionActorDescUtils.h
+++ b/Engine/Source/Runtime/Engine/Public/WorldPartition/WorldPartitionActorDescUtils.h
@@ -28,7 +28,7 @@ struct ENGINE_API FWorldPartitionActorDescUtils
 	/** 
 	 * Creates a valid actor descriptor from the provided asset data.
 	 * @param InAssetData	The asset data to look into.
-	 * @return				Actor descriptor retrieved from the provided asset data..
+	 * @return				Actor descriptor retrieved from the provided asset data.
 	 */
 	static TUniquePtr<FWorldPartitionActorDesc> GetActorDescriptorFromAssetData(const FAssetData& InAssetData);
 
@@ -41,10 +41,17 @@ struct ENGINE_API FWorldPartitionActorDescUtils
 
 	/** 
 	 * Update an actor descriptor with new values coming from the provided actor.
-	 * @param InActor		The actor to update from..
+	 * @param InActor		The actor to update from.
 	 * @param OutActorDesc	Actor descriptor unique pointer that will get updated with the new actor descriptor.
 	 */
-	static void UpdateActorDescriptorFomActor(const AActor* InActor, TUniquePtr<FWorldPartitionActorDesc>& OutActorDesc);
+	static void UpdateActorDescriptorFromActor(const AActor* InActor, TUniquePtr<FWorldPartitionActorDesc>& OutActorDesc);
+
+	/** 
+	 * Update an actor descriptor with new values coming from the provided actor descriptor.
+	 * @param InActorDesc	The actor descriptor to update from.
+	 * @param OutActorDesc	Actor descriptor unique pointer that will get updated with the new actor descriptor.
+	 */
+	static void UpdateActorDescriptorFromActorDescriptor(TUniquePtr<FWorldPartitionActorDesc>& InActorDesc, TUniquePtr<FWorldPartitionActorDesc>& OutActorDesc);
 
 	/** 
 	 * Replaces the actor descriptor's actor pointer with the provided new actor pointer.
