World Partition: add a button to remove world partition from a world, only available if streaming is disabled.

diff --git a/Engine/Source/Editor/WorldPartitionEditor/Private/WorldPartition/Customizations/WorldPartitionDetailsCustomization.cpp b/Engine/Source/Editor/WorldPartitionEditor/Private/WorldPartition/Customizations/WorldPartitionDetailsCustomization.cpp
--- a/Engine/Source/Editor/WorldPartitionEditor/Private/WorldPartition/Customizations/WorldPartitionDetailsCustomization.cpp
+++ b/Engine/Source/Editor/WorldPartitionEditor/Private/WorldPartition/Customizations/WorldPartitionDetailsCustomization.cpp
@@ -53,6 +53,45 @@ void FWorldPartitionDetails::CustomizeDetails(IDetailLayoutBuilder& InDetailBuil
 		]
 		.Visibility(TAttribute<EVisibility>::CreateLambda([this]() { return WorldPartition.IsValid() && WorldPartition->SupportsStreaming() ? EVisibility::Visible : EVisibility::Hidden; }));
 
+	// Disable world partition button
+	WorldPartitionCategory.AddCustomRow(LOCTEXT("DefaultWorldPartitionSettingsRow", "DefaultWorldPartitionSettings"), true)
+		.NameContent()
+		[
+			SNew(STextBlock)
+			.Text(LOCTEXT("DisableWorldPartition", "Disable World Partition"))
+			.ToolTipText(LOCTEXT("DisableWorldPartition_ToolTip", "Disable World Partition for this world"))
+			.Font(IDetailLayoutBuilder::GetDetailFont())
+		]
+		.ValueContent()
+		[
+			SNew(SHorizontalBox)
+			+ SHorizontalBox::Slot()
+			[
+				SNew(SButton)
+				.VAlign(VAlign_Center)
+				.HAlign(HAlign_Center)
+				.IsEnabled_Lambda([this]()
+				{
+					return !WorldPartition->IsStreamingEnabled();
+				})
+				.OnClicked_Lambda([this]()
+				{
+					if (FMessageDialog::Open(EAppMsgType::YesNo, LOCTEXT("RemoveWorldPartitionConfiormation", "You are about to remove world partition from the current level. Continue?")) == EAppReturnType::Yes)
+					{
+						FScopedTransaction Transaction(LOCTEXT("RemoveWorldPartition", "Remove World Partition"));
+						UWorldPartition::RemoveWorldPartition(World->GetWorldSettings());
+					}
+					return FReply::Handled();
+				})
+				[
+					SNew(STextBlock)
+					.Text(LOCTEXT("DisableButtom", "Disable"))
+					.ToolTipText(LOCTEXT("DisableWorldPartition_ToolTip", "Disable World Partition for this world, only works if streaming is disabled."))
+					.Font(IDetailLayoutBuilder::GetDetailFont())
+				]
+			]
+		];
+
 	// Runtime hash class selector
 	WorldPartitionCategory.AddCustomRow(LOCTEXT("RuntimeHashClass", "Runtime Hash Class"), false)
 		.RowTag(TEXT("RuntimeHashClass"))
diff --git a/Engine/Source/Runtime/Engine/Private/WorldPartition/WorldPartition.cpp b/Engine/Source/Runtime/Engine/Private/WorldPartition/WorldPartition.cpp
--- a/Engine/Source/Runtime/Engine/Private/WorldPartition/WorldPartition.cpp
+++ b/Engine/Source/Runtime/Engine/Private/WorldPartition/WorldPartition.cpp
@@ -1005,6 +1005,36 @@ UWorldPartition* UWorldPartition::CreateOrRepairWorldPartition(AWorldSettings* W
 	return WorldPartition;
 }
 
+bool UWorldPartition::RemoveWorldPartition(AWorldSettings* WorldSettings)
+{
+	if (UWorldPartition* WorldPartition = WorldSettings->GetWorldPartition())
+	{
+		if (!WorldPartition->IsStreamingEnabled())
+		{
+			WorldSettings->Modify();
+			
+			ULevel* PersistentLevel = WorldSettings->GetLevel();
+			for (AActor* Actor : PersistentLevel->Actors)
+			{
+				if (Actor && (Cast<AWorldDataLayers>(Actor) || Cast<AWorldPartitionMiniMap>(Actor)))
+				{
+					Actor->Destroy();
+				}
+			}
+
+			WorldPartition->Uninitialize();
+			WorldSettings->SetWorldPartition(nullptr);
+
+			if (WorldPartition->WorldPartitionEditor)
+			{
+				WorldPartition->WorldPartitionEditor->Reconstruct();
+			}
+			
+			return true;
+		}
+	}
+	return false;
+}
 #endif
 
 const TArray<FWorldPartitionStreamingSource>& UWorldPartition::GetStreamingSources() const
diff --git a/Engine/Source/Runtime/Engine/Private/WorldSettings.cpp b/Engine/Source/Runtime/Engine/Private/WorldSettings.cpp
--- a/Engine/Source/Runtime/Engine/Private/WorldSettings.cpp
+++ b/Engine/Source/Runtime/Engine/Private/WorldSettings.cpp
@@ -222,8 +222,7 @@ UWorldPartition* AWorldSettings::GetWorldPartition() const
 
 void AWorldSettings::SetWorldPartition(UWorldPartition* InWorldPartition)
 {
-	check(!IsValid(WorldPartition));
-	check(InWorldPartition);
+	check(!InWorldPartition || !WorldPartition);
 	WorldPartition = InWorldPartition;
 	ApplyWorldPartitionForcedSettings();
 }
diff --git a/Engine/Source/Runtime/Engine/Public/WorldPartition/WorldPartition.h b/Engine/Source/Runtime/Engine/Public/WorldPartition/WorldPartition.h
--- a/Engine/Source/Runtime/Engine/Public/WorldPartition/WorldPartition.h
+++ b/Engine/Source/Runtime/Engine/Public/WorldPartition/WorldPartition.h
@@ -93,6 +93,7 @@ class ENGINE_API UWorldPartition final : public UObject, public FActorDescContai
 public:
 #if WITH_EDITOR
 	static UWorldPartition* CreateOrRepairWorldPartition(AWorldSettings* WorldSettings, TSubclassOf<UWorldPartitionEditorHash> EditorHashClass = nullptr, TSubclassOf<UWorldPartitionRuntimeHash> RuntimeHashClass = nullptr);
+	static bool RemoveWorldPartition(AWorldSettings* WorldSettings);
 #endif
 
 	DECLARE_MULTICAST_DELEGATE_OneParam(FWorldPartitionInitializeDelegate, UWorldPartition*);
