compile fix

diff --git a/Engine/Source/Editor/WorldPartitionEditor/Private/WorldPartition/Customizations/WorldPartitionDetailsCustomization.cpp b/Engine/Source/Editor/WorldPartitionEditor/Private/WorldPartition/Customizations/WorldPartitionDetailsCustomization.cpp
--- a/Engine/Source/Editor/WorldPartitionEditor/Private/WorldPartition/Customizations/WorldPartitionDetailsCustomization.cpp
+++ b/Engine/Source/Editor/WorldPartitionEditor/Private/WorldPartition/Customizations/WorldPartitionDetailsCustomization.cpp
@@ -4,9 +4,11 @@
 #include "WorldPartition/WorldPartition.h"
 #include "WorldPartition/WorldPartitionRuntimeHash.h"
 #include "PropertyCustomizationHelpers.h"
+#include "ScopedTransaction.h"
 #include "DetailLayoutBuilder.h"
 #include "DetailCategoryBuilder.h"
 #include "DetailWidgetRow.h"
+#include "Engine/World.h"
 #include "Misc/MessageDialog.h"
 #include "Widgets/Layout/SBorder.h"
 #include "Widgets/Text/STextBlock.h"
@@ -76,10 +78,10 @@ void FWorldPartitionDetails::CustomizeDetails(IDetailLayoutBuilder& InDetailBuil
 				})
 				.OnClicked_Lambda([this]()
 				{
-					if (FMessageDialog::Open(EAppMsgType::YesNo, LOCTEXT("RemoveWorldPartitionConfiormation", "You are about to remove world partition from the current level. Continue?")) == EAppReturnType::Yes)
+					if (FMessageDialog::Open(EAppMsgType::YesNo, LOCTEXT("RemoveWorldPartitionConfirmation", "You are about to remove world partition from the current level. Continue?")) == EAppReturnType::Yes)
 					{
 						FScopedTransaction Transaction(LOCTEXT("RemoveWorldPartition", "Remove World Partition"));
-						UWorldPartition::RemoveWorldPartition(World->GetWorldSettings());
+						UWorldPartition::RemoveWorldPartition(WorldPartition->GetWorld()->GetWorldSettings());
 					}
 					return FReply::Handled();
 				})
