Fix localization warnings.

diff --git a/Engine/Source/Editor/WorldPartitionEditor/Private/WorldPartition/Customizations/WorldPartitionDetailsCustomization.cpp b/Engine/Source/Editor/WorldPartitionEditor/Private/WorldPartition/Customizations/WorldPartitionDetailsCustomization.cpp
--- a/Engine/Source/Editor/WorldPartitionEditor/Private/WorldPartition/Customizations/WorldPartitionDetailsCustomization.cpp
+++ b/Engine/Source/Editor/WorldPartitionEditor/Private/WorldPartition/Customizations/WorldPartitionDetailsCustomization.cpp
@@ -61,7 +61,7 @@ void FWorldPartitionDetails::CustomizeDetails(IDetailLayoutBuilder& InDetailBuil
 		[
 			SNew(STextBlock)
 			.Text(LOCTEXT("DisableWorldPartition", "Disable World Partition"))
-			.ToolTipText(LOCTEXT("DisableWorldPartition_ToolTip", "Disable World Partition for this world"))
+			.ToolTipText(LOCTEXT("DisableWorldPartition_ToolTip", "Disable World Partition for this world, only works if streaming is disabled."))
 			.Font(IDetailLayoutBuilder::GetDetailFont())
 		]
 		.ValueContent()
