World Partition: When invoking external commandlets (for HLODs/Minimap), make sure to provide the full set of command line arguments as provided for the editor launch

diff --git a/Engine/Source/Editor/WorldPartitionEditor/Private/WorldPartitionEditorModule.cpp b/Engine/Source/Editor/WorldPartitionEditor/Private/WorldPartitionEditorModule.cpp
--- a/Engine/Source/Editor/WorldPartitionEditor/Private/WorldPartitionEditorModule.cpp
+++ b/Engine/Source/Editor/WorldPartitionEditor/Private/WorldPartitionEditorModule.cpp
@@ -370,18 +370,19 @@ void FWorldPartitionEditorModule::RunCommandletAsExternalProcess(const FString&
 
 	FString CurrentExecutableName = FPlatformProcess::ExecutablePath();
 
-	TArray<FString> AdditionalArgs;
-	OnExecuteCommandletEvent.Broadcast(AdditionalArgs);
-
 	// Try to provide complete Path, if we can't try with project name
 	FString ProjectPath = FPaths::IsProjectFilePathSet() ? FPaths::GetProjectFilePath() : FApp::GetProjectName();
+		
+	TArray<FString> CommandletArgsArray;
+	CommandletArgsArray.Add(TEXT("\"") + ProjectPath + TEXT('"'));
+	CommandletArgsArray.Add(TEXT("-BaseDir=\"") + FString(FPlatformProcess::BaseDir()) + TEXT('"'));
+	CommandletArgsArray.Add(TEXT("-Unattended"));
+	CommandletArgsArray.Add(InCommandletArgs);
+
+	OnExecuteCommandletEvent.Broadcast(CommandletArgsArray);
 
 	FString Arguments;
-	Arguments += TEXT("\"") + ProjectPath + TEXT('"');
-	Arguments += TEXT(" -BaseDir=\"") + FString(FPlatformProcess::BaseDir()) + TEXT('"');
-	Arguments += TEXT(" -Unattended");
-	Arguments += TEXT(" ") + InCommandletArgs;
-	for (const FString& AdditionalArg : AdditionalArgs)
+	for (const FString& AdditionalArg : CommandletArgsArray)
 	{
 		Arguments += " ";
 		Arguments += AdditionalArg;
