Fix editor hang when the World Partition editor window size has zero pixel width/height.

diff --git a/Engine/Source/Editor/WorldPartitionEditor/Private/WorldPartition/SWorldPartitionEditorGridSpatialHash.cpp b/Engine/Source/Editor/WorldPartitionEditor/Private/WorldPartition/SWorldPartitionEditorGridSpatialHash.cpp
--- a/Engine/Source/Editor/WorldPartitionEditor/Private/WorldPartition/SWorldPartitionEditorGridSpatialHash.cpp
+++ b/Engine/Source/Editor/WorldPartitionEditor/Private/WorldPartition/SWorldPartitionEditorGridSpatialHash.cpp
@@ -33,11 +33,12 @@ int64 SWorldPartitionEditorGridSpatialHash::GetSelectionSnap(float& FadeRatio, f
 	int64 EffectiveCellSize = EditorSpatialHash->CellSize;
 	for (;;)
 	{
-		CellScreenSize = WorldToScreen.TransformVector(FVector2D(EffectiveCellSize, EffectiveCellSize)).X;
+		CellScreenSize = WorldToScreen.TransformVector(FVector2D(EffectiveCellSize, EffectiveCellSize)).GetMax();
 
-		if (CellScreenSize > WantedCellScreenSize)
+		// CellScreenSize can be zero when the window screen rect is too small
+		if ((CellScreenSize <= 0.0) || (CellScreenSize > WantedCellScreenSize))
 		{
-			FadeRatio = (CellScreenSize - WantedCellScreenSize) / WantedCellScreenSize;
+			FadeRatio = FMath::Clamp((CellScreenSize - WantedCellScreenSize) / WantedCellScreenSize, 0.0f, 1.0f);
 			break;
 		}
 	
