World Partition: Always ask user to save when performing build operations, even for those that don't require the map to be unloaded

diff --git a/Engine/Source/Editor/WorldPartitionEditor/Private/WorldPartitionEditorModule.cpp b/Engine/Source/Editor/WorldPartitionEditor/Private/WorldPartitionEditorModule.cpp
--- a/Engine/Source/Editor/WorldPartitionEditor/Private/WorldPartitionEditorModule.cpp
+++ b/Engine/Source/Editor/WorldPartitionEditor/Private/WorldPartitionEditorModule.cpp
@@ -301,15 +301,20 @@ void FWorldPartitionEditorModule::OnConvertMap()
 	}
 }
 
-static bool UnloadCurrentMap(bool bAskSaveContentPackages, FString& MapPackageName)
+static bool AskSaveDirtyPackages(const bool bSaveContentPackages)
 {
-	UPackage* WorldPackage = FindPackage(nullptr, *MapPackageName);
-
 	// Ask user to save dirty packages
-	if (!FEditorFileUtils::SaveDirtyPackages(/*bPromptUserToSave=*/true, /*bSaveMapPackages=*/true, bAskSaveContentPackages))
-	{
-		return false;
-	}
+	const bool bPromptUserToSave = true;
+	const bool bSaveMapPackages = true;
+	const bool bFastSave = false;
+	const bool bNotifyNoPackagesSaved = false;
+	const bool bCanBeDeclined = false;
+	return FEditorFileUtils::SaveDirtyPackages(bPromptUserToSave, bSaveMapPackages, bSaveContentPackages, bFastSave, bNotifyNoPackagesSaved, bCanBeDeclined);
+}
+
+static bool UnloadCurrentMap(FString& MapPackageName)
+{
+	UPackage* WorldPackage = FindPackage(nullptr, *MapPackageName);
 
 	// Make sure we handle the case where the world package was renamed on save (for temp world for example)
 	if (WorldPackage)
@@ -480,7 +485,13 @@ bool FWorldPartitionEditorModule::ConvertMap(const FString& InLongPackageName)
 
 	if (ConvertDialog->ClickedOk())
 	{
-		if (!UnloadCurrentMap(/*bAskSaveContentPackages=*/false, DefaultConvertOptions->LongPackageName))
+		// Ask user to save dirty packages
+		if (!AskSaveDirtyPackages(/*bAskSaveContentPackages=*/false))
+		{
+			return false;
+		}
+
+		if (!UnloadCurrentMap(DefaultConvertOptions->LongPackageName))
 		{
 			return false;
 		}
@@ -601,8 +612,15 @@ bool FWorldPartitionEditorModule::Build(const FRunBuilderParams& InParams)
 	FRunBuilderParams ParamsCopy(InParams);
 	OnPreExecuteCommandletEvent.Broadcast(ParamsCopy);
 
+	// Ask user to save dirty packages
+	if (!AskSaveDirtyPackages(/*bAskSaveContentPackages=*/true))
+	{
+		return false;
+	}
+
+	// Unload map if required
 	FString MapPackage = ParamsCopy.World->GetPackage()->GetName();
-	if (ParamsCopy.bUnloadMap && !UnloadCurrentMap(/*bAskSaveContentPackages=*/true, MapPackage))
+	if (ParamsCopy.bUnloadMap && !UnloadCurrentMap(MapPackage))
 	{
 		return false;
 	}
