World Partition: fix for rare crash when closing the world partition editor.

diff --git a/Engine/Source/Editor/WorldPartitionEditor/Private/WorldPartition/SWorldPartitionEditor.cpp b/Engine/Source/Editor/WorldPartitionEditor/Private/WorldPartition/SWorldPartitionEditor.cpp
--- a/Engine/Source/Editor/WorldPartitionEditor/Private/WorldPartition/SWorldPartitionEditor.cpp
+++ b/Engine/Source/Editor/WorldPartitionEditor/Private/WorldPartition/SWorldPartitionEditor.cpp
@@ -40,9 +40,12 @@ SWorldPartitionEditor::~SWorldPartitionEditor()
 	{
 		if (UWorldPartition* WorldPartition = World->GetWorldPartition())
 		{
-			check(WorldPartition->World == World);
-			check(WorldPartition->WorldPartitionEditor == this);
-			WorldPartition->WorldPartitionEditor = nullptr;
+			check(WorldPartition->World == World);			
+			if (WorldPartition->WorldPartitionEditor)
+			{
+				check(WorldPartition->WorldPartitionEditor == this);
+				WorldPartition->WorldPartitionEditor = nullptr;
+			}
 		}
 	}
 }
