Add a scope for ContentBundleActivation
Instantiate the scope to activate a ContentBundle and deactivate it upon scope deletion.

diff --git a/Engine/Source/Editor/WorldPartitionEditor/Public/WorldPartition/ContentBundle/ContentBundleEditorSubsystem.h b/Engine/Source/Editor/WorldPartitionEditor/Public/WorldPartition/ContentBundle/ContentBundleEditorSubsystem.h
--- a/Engine/Source/Editor/WorldPartitionEditor/Public/WorldPartition/ContentBundle/ContentBundleEditorSubsystem.h
+++ b/Engine/Source/Editor/WorldPartitionEditor/Public/WorldPartition/ContentBundle/ContentBundleEditorSubsystem.h
@@ -118,7 +118,7 @@ public:
 
 	TArray<TSharedPtr<FContentBundleEditor>> GetEditorContentBundles();
 	TSharedPtr<FContentBundleEditor> GetEditorContentBundle(const UContentBundleDescriptor* ContentBundleDescriptor) const;
-	TSharedPtr<FContentBundleEditor> GetEditorContentBundle(const FGuid& ContentBundleGuid) const;
+	TSharedPtr<FContentBundleEditor> GetEditorContentBundle(const FGuid& ContentBundleGuid) const override;
 
 	void SelectActors(FContentBundleEditor& EditorContentBundle);
 	void DeselectActors(FContentBundleEditor& EditorContentBundle);
@@ -126,10 +126,10 @@ public:
 	void ReferenceAllActors(FContentBundleEditor& EditorContentBundle);
 	void UnreferenceAllActors(FContentBundleEditor& EditorContentBundle);
 
-	bool IsEditingContentBundle() const;
+	bool IsEditingContentBundle() const override;
 	bool IsEditingContentBundle(const FGuid& ContentBundleGuid) const;
-	bool ActivateContentBundleEditing(TSharedPtr<FContentBundleEditor>& ContentBundleEditor) const;
-	bool DeactivateContentBundleEditing(TSharedPtr<FContentBundleEditor>& ContentBundleEditor) const;
+	bool ActivateContentBundleEditing(TSharedPtr<FContentBundleEditor>& ContentBundleEditor) const override;
+	bool DeactivateContentBundleEditing(TSharedPtr<FContentBundleEditor>& ContentBundleEditor) const override;
 	bool IsContentBundleEditingActivated(TSharedPtr<FContentBundleEditor>& ContentBundleEditor) const;
 
 	DECLARE_EVENT_OneParam(UContentBundleEditorSubsystem, FOnContentBundleChanged, const FContentBundleEditor*);
diff --git a/Engine/Source/Runtime/Engine/Public/WorldPartition/ContentBundle/ContentBundleActivationScope.h b/Engine/Source/Runtime/Engine/Public/WorldPartition/ContentBundle/ContentBundleActivationScope.h
new file mode 100644
--- /dev/null
+++ b/Engine/Source/Runtime/Engine/Public/WorldPartition/ContentBundle/ContentBundleActivationScope.h
@@ -0,0 +1,44 @@
+// Copyright Epic Games, Inc. All Rights Reserved.
+#pragma once
+
+#include "CoreMinimal.h"
+#include "WorldPartition/ContentBundle/ContentBundleEditorSubsystemInterface.h"
+
+class FContentBundleEditor;
+
+#if WITH_EDITOR
+
+class FContentBundleActivationScope
+{
+public:
+	FContentBundleActivationScope(FGuid InContentBundleGuid)
+	{
+		if (InContentBundleGuid.IsValid())
+		{
+			IContentBundleEditorSubsystemInterface* ContentBundleEditorSubsystem = IContentBundleEditorSubsystemInterface::Get();
+			if (TSharedPtr<FContentBundleEditor> ContentBundleEditor = ContentBundleEditorSubsystem->GetEditorContentBundle(InContentBundleGuid))
+			{
+				check(!ContentBundleEditorSubsystem->IsEditingContentBundle());
+
+				ContentBundleEditorSubsystem->ActivateContentBundleEditing(ContentBundleEditor);
+				ContentBundleGuid = InContentBundleGuid;
+			}
+		}
+	}
+	~FContentBundleActivationScope()
+	{
+		if (ContentBundleGuid.IsValid())
+		{
+			IContentBundleEditorSubsystemInterface* ContentBundleEditorSubsystem = IContentBundleEditorSubsystemInterface::Get();
+			if (TSharedPtr<FContentBundleEditor> ContentBundleEditor = ContentBundleEditorSubsystem->GetEditorContentBundle(ContentBundleGuid))
+			{
+				ContentBundleEditorSubsystem->DeactivateContentBundleEditing(ContentBundleEditor);
+			}
+		}
+	}
+
+private:
+	FGuid ContentBundleGuid;
+};
+
+#endif
diff --git a/Engine/Source/Runtime/Engine/Public/WorldPartition/ContentBundle/ContentBundleEditorSubsystemInterface.h b/Engine/Source/Runtime/Engine/Public/WorldPartition/ContentBundle/ContentBundleEditorSubsystemInterface.h
--- a/Engine/Source/Runtime/Engine/Public/WorldPartition/ContentBundle/ContentBundleEditorSubsystemInterface.h
+++ b/Engine/Source/Runtime/Engine/Public/WorldPartition/ContentBundle/ContentBundleEditorSubsystemInterface.h
@@ -17,6 +17,12 @@ public:
 	virtual void NotifyContentBundleRemovedContent(const FContentBundleEditor* ContentBundle) = 0;
 	virtual void NotifyContentBundleChanged(const FContentBundleEditor* ContentBundle) = 0;
 
+	virtual TSharedPtr<FContentBundleEditor> GetEditorContentBundle(const FGuid& ContentBundleGuid) const = 0;
+	
+	virtual bool IsEditingContentBundle() const = 0;
+	virtual bool ActivateContentBundleEditing(TSharedPtr<FContentBundleEditor>& ContentBundleEditor) const = 0;
+	virtual bool DeactivateContentBundleEditing(TSharedPtr<FContentBundleEditor>& ContentBundleEditor) const = 0;
+
 protected:
 	static void SetInstance(IContentBundleEditorSubsystemInterface* InInstance);
 	static IContentBundleEditorSubsystemInterface* Instance;
