World Partition:
- Added an advanced mode that can be triggered on or off by projects.
- Only show advanced streaming properties when enabling advanced mode.

diff --git a/Engine/Source/Editor/WorldPartitionEditor/Private/WorldPartition/Customizations/WorldPartitionDetailsCustomization.cpp b/Engine/Source/Editor/WorldPartitionEditor/Private/WorldPartition/Customizations/WorldPartitionDetailsCustomization.cpp
--- a/Engine/Source/Editor/WorldPartitionEditor/Private/WorldPartition/Customizations/WorldPartitionDetailsCustomization.cpp
+++ b/Engine/Source/Editor/WorldPartitionEditor/Private/WorldPartition/Customizations/WorldPartitionDetailsCustomization.cpp
@@ -3,6 +3,7 @@
 #include "WorldPartitionDetailsCustomization.h"
 #include "WorldPartition/WorldPartition.h"
 #include "WorldPartition/WorldPartitionRuntimeHash.h"
+#include "WorldPartition/WorldPartitionEditorSettings.h"
 #include "PropertyCustomizationHelpers.h"
 #include "ScopedTransaction.h"
 #include "DetailLayoutBuilder.h"
@@ -56,64 +57,86 @@ void FWorldPartitionDetails::CustomizeDetails(IDetailLayoutBuilder& InDetailBuil
 		.Visibility(TAttribute<EVisibility>::CreateLambda([this]() { return WorldPartition.IsValid() && WorldPartition->SupportsStreaming() ? EVisibility::Visible : EVisibility::Hidden; }));
 
 	// Disable world partition button
-	WorldPartitionCategory.AddCustomRow(LOCTEXT("DefaultWorldPartitionSettingsRow", "DefaultWorldPartitionSettings"), true)
-		.NameContent()
-		[
-			SNew(STextBlock)
-			.Text(LOCTEXT("DisableWorldPartition", "Disable World Partition"))
-			.ToolTipText(LOCTEXT("DisableWorldPartition_ToolTip", "Disable World Partition for this world, only works if streaming is disabled."))
-			.Font(IDetailLayoutBuilder::GetDetailFont())
-		]
-		.ValueContent()
-		[
-			SNew(SHorizontalBox)
-			+ SHorizontalBox::Slot()
+	if (GetDefault<UWorldPartitionEditorSettings>()->bAdvancedMode)
+	{
+		// Can be used by level instance checkbox
+		WorldPartitionCategory.AddCustomRow(LOCTEXT("CanBeUsedByLevelInstance", "Can Be Used By Level Instance"), false)
+			.RowTag(TEXT("CanBeUsedByLevelInstance"))
+			.NameContent()
 			[
-				SNew(SButton)
-				.VAlign(VAlign_Center)
-				.HAlign(HAlign_Center)
-				.IsEnabled_Lambda([this]()
-				{
-					return !WorldPartition->IsStreamingEnabled();
-				})
-				.OnClicked_Lambda([this]()
-				{
-					if (FMessageDialog::Open(EAppMsgType::YesNo, LOCTEXT("RemoveWorldPartitionConfirmation", "You are about to remove world partition from the current level. Continue?")) == EAppReturnType::Yes)
-					{
-						FScopedTransaction Transaction(LOCTEXT("RemoveWorldPartition", "Remove World Partition"));
-						UWorldPartition::RemoveWorldPartition(WorldPartition->GetWorld()->GetWorldSettings());
-					}
-					return FReply::Handled();
-				})
+				SNew(STextBlock)
+				.Text(LOCTEXT("WorldPartitionCanBeUsedByLevelInstance", "Can Be Used By Level Instance"))
+				.ToolTipText(LOCTEXT("WorldPartitionCanBeUsedByLevelInstance_ToolTip", "Whether Level Instance can reference this partition."))
+				.Font(IDetailLayoutBuilder::GetDetailFont())
+			]
+			.ValueContent()
+			[
+				SNew(SCheckBox)
+				.IsChecked(MakeAttributeLambda([this]() { return WorldPartition.IsValid() && WorldPartition->CanBeUsedByLevelInstance() ? ECheckBoxState::Checked : ECheckBoxState::Unchecked; }))
+				.OnCheckStateChanged(this, &FWorldPartitionDetails::HandleWorldPartitionCanBeUsedByLevelInstanceChanged)
+			]
+			.Visibility(TAttribute<EVisibility>::CreateLambda([this]() { return WorldPartition.IsValid() && !WorldPartition->IsStreamingEnabled() ? EVisibility::Visible : EVisibility::Hidden; }));
+
+		// Default world partition settings
+		WorldPartitionCategory.AddCustomRow(LOCTEXT("DefaultWorldPartitionSettingsRow", "DefaultWorldPartitionSettings"), true)
+			.NameContent()
+			[
+				SNew(STextBlock)
+				.Text(LOCTEXT("DisableWorldPartition", "Disable World Partition"))
+				.ToolTipText(LOCTEXT("DisableWorldPartition_ToolTip", "Disable World Partition for this world, only works if streaming is disabled."))
+				.Font(IDetailLayoutBuilder::GetDetailFont())
+			]
+			.ValueContent()
+			[
+				SNew(SHorizontalBox)
+				+ SHorizontalBox::Slot()
 				[
-					SNew(STextBlock)
-					.Text(LOCTEXT("DisableButtom", "Disable"))
-					.ToolTipText(LOCTEXT("DisableWorldPartition_ToolTip", "Disable World Partition for this world, only works if streaming is disabled."))
-					.Font(IDetailLayoutBuilder::GetDetailFont())
+					SNew(SButton)
+					.VAlign(VAlign_Center)
+					.HAlign(HAlign_Center)
+					.IsEnabled_Lambda([this]()
+					{
+						return !WorldPartition->IsStreamingEnabled();
+					})
+					.OnClicked_Lambda([this]()
+					{
+						if (FMessageDialog::Open(EAppMsgType::YesNo, LOCTEXT("RemoveWorldPartitionConfirmation", "You are about to remove world partition from the current level. Continue?")) == EAppReturnType::Yes)
+						{
+							FScopedTransaction Transaction(LOCTEXT("RemoveWorldPartition", "Remove World Partition"));
+							UWorldPartition::RemoveWorldPartition(WorldPartition->GetWorld()->GetWorldSettings());
+						}
+						return FReply::Handled();
+					})
+					[
+						SNew(STextBlock)
+						.Text(LOCTEXT("DisableButtom", "Disable"))
+						.ToolTipText(LOCTEXT("DisableWorldPartition_ToolTip", "Disable World Partition for this world, only works if streaming is disabled."))
+						.Font(IDetailLayoutBuilder::GetDetailFont())
+					]
 				]
-			]
-		];
+			];
 
-	// Runtime hash class selector
-	WorldPartitionCategory.AddCustomRow(LOCTEXT("RuntimeHashClass", "Runtime Hash Class"), false)
-		.RowTag(TEXT("RuntimeHashClass"))
-		.NameContent()
-		[
-			SNew(STextBlock)
-			.Text(LOCTEXT("WorldPartitionRuntimeHashClass", "Runtime Hash Class"))
-			.ToolTipText(LOCTEXT("WorldPartitionRuntimeHashClass_ToolTip", "Set the world partition runtime hash class."))
-			.Font(IDetailLayoutBuilder::GetDetailFont())
-		]
-		.ValueContent()
-		[
-			SNew(SClassPropertyEntryBox)
-			.MetaClass(UWorldPartitionRuntimeHash::StaticClass())
-			.AllowNone(false)
-			.HideViewOptions(true)
-			.SelectedClass_Lambda([this]() { return RuntimeHashClass; })
-			.OnSetClass_Lambda([this](const UClass* Class) { HandleWorldPartitionRuntimeHashClassChanged(Class); })
-		]
-		.Visibility(TAttribute<EVisibility>::CreateLambda([this]() { return WorldPartition.IsValid() && WorldPartition->IsStreamingEnabled() ? EVisibility::Visible : EVisibility::Hidden; }));
+		// Runtime hash class selector
+		WorldPartitionCategory.AddCustomRow(LOCTEXT("RuntimeHashClass", "Runtime Hash Class"), false)
+			.RowTag(TEXT("RuntimeHashClass"))
+			.NameContent()
+			[
+				SNew(STextBlock)
+				.Text(LOCTEXT("WorldPartitionRuntimeHashClass", "Runtime Hash Class"))
+				.ToolTipText(LOCTEXT("WorldPartitionRuntimeHashClass_ToolTip", "Set the world partition runtime hash class."))
+				.Font(IDetailLayoutBuilder::GetDetailFont())
+			]
+			.ValueContent()
+			[
+				SNew(SClassPropertyEntryBox)
+				.MetaClass(UWorldPartitionRuntimeHash::StaticClass())
+				.AllowNone(false)
+				.HideViewOptions(true)
+				.SelectedClass_Lambda([this]() { return RuntimeHashClass; })
+				.OnSetClass_Lambda([this](const UClass* Class) { HandleWorldPartitionRuntimeHashClassChanged(Class); })
+			]
+			.Visibility(TAttribute<EVisibility>::CreateLambda([this]() { return WorldPartition.IsValid() && WorldPartition->IsStreamingEnabled() ? EVisibility::Visible : EVisibility::Hidden; }));
+	}
 
 	// Runtime hash properties
 	if (WorldPartition->RuntimeHash)
@@ -175,6 +198,11 @@ void FWorldPartitionDetails::HandleWorldPartitionEnableStreamingChanged(ECheckBo
 	}
 }
 
+void FWorldPartitionDetails::HandleWorldPartitionCanBeUsedByLevelInstanceChanged(ECheckBoxState InCheckState)
+{
+	WorldPartition->SetCanBeUsedByLevelInstance((InCheckState == ECheckBoxState::Checked));
+}
+
 void FWorldPartitionDetails::HandleWorldPartitionRuntimeHashClassChanged(const UClass* InRuntimeHashClass)
 {
 	check(InRuntimeHashClass);
diff --git a/Engine/Source/Editor/WorldPartitionEditor/Private/WorldPartition/Customizations/WorldPartitionDetailsCustomization.h b/Engine/Source/Editor/WorldPartitionEditor/Private/WorldPartition/Customizations/WorldPartitionDetailsCustomization.h
--- a/Engine/Source/Editor/WorldPartitionEditor/Private/WorldPartition/Customizations/WorldPartitionDetailsCustomization.h
+++ b/Engine/Source/Editor/WorldPartitionEditor/Private/WorldPartition/Customizations/WorldPartitionDetailsCustomization.h
@@ -28,6 +28,9 @@ private:
 	// Callback for changes in the world partition enable streaming.
 	void HandleWorldPartitionEnableStreamingChanged(ECheckBoxState InCheckState);
 
+	// Callback for changes in the world partition can be used by level instance.
+	void HandleWorldPartitionCanBeUsedByLevelInstanceChanged(ECheckBoxState InCheckState);
+
 	// Callback for changes in the world partition runtime hash class.
 	void HandleWorldPartitionRuntimeHashClassChanged(const UClass* InRuntimeHashClass);
 
diff --git a/Engine/Source/Editor/WorldPartitionEditor/Private/WorldPartitionEditorModule.cpp b/Engine/Source/Editor/WorldPartitionEditor/Private/WorldPartitionEditorModule.cpp
--- a/Engine/Source/Editor/WorldPartitionEditor/Private/WorldPartitionEditorModule.cpp
+++ b/Engine/Source/Editor/WorldPartitionEditor/Private/WorldPartitionEditorModule.cpp
@@ -281,6 +281,16 @@ void FWorldPartitionEditorModule::SetDisableBugIt(bool bInDisableBugIt)
 	GetMutableDefault<UWorldPartitionEditorSettings>()->bDisableBugIt = bInDisableBugIt;
 }
 
+bool FWorldPartitionEditorModule::GetAdvancedMode() const
+{
+	return GetDefault<UWorldPartitionEditorSettings>()->bAdvancedMode;
+}
+
+void FWorldPartitionEditorModule::SetAdvancedMode(bool bInAdvancedMode)
+{
+	GetMutableDefault<UWorldPartitionEditorSettings>()->bAdvancedMode = bInAdvancedMode;
+}
+
 void FWorldPartitionEditorModule::OnConvertMap()
 {
 	IContentBrowserSingleton& ContentBrowserSingleton = FModuleManager::LoadModuleChecked<FContentBrowserModule>("ContentBrowser").Get();
@@ -748,6 +758,7 @@ UWorldPartitionEditorSettings::UWorldPartitionEditorSettings()
 	bDisableLoadingInEditor = false;
 	bDisablePIE = false;
 	bDisableBugIt = false;
+	bAdvancedMode = true;
 }
 
 FString UWorldPartitionConvertOptions::ToCommandletArgs() const
diff --git a/Engine/Source/Editor/WorldPartitionEditor/Public/WorldPartition/WorldPartitionEditorSettings.h b/Engine/Source/Editor/WorldPartitionEditor/Public/WorldPartition/WorldPartitionEditorSettings.h
--- a/Engine/Source/Editor/WorldPartitionEditor/Public/WorldPartition/WorldPartitionEditorSettings.h
+++ b/Engine/Source/Editor/WorldPartitionEditor/Public/WorldPartition/WorldPartitionEditorSettings.h
@@ -27,4 +27,5 @@ public:
 	bool bDisableLoadingInEditor;
 	bool bDisableBugIt;
 	bool bDisablePIE;
+	bool bAdvancedMode;
 };
\ No newline at end of file
diff --git a/Engine/Source/Editor/WorldPartitionEditor/Public/WorldPartitionEditorModule.h b/Engine/Source/Editor/WorldPartitionEditor/Public/WorldPartitionEditorModule.h
--- a/Engine/Source/Editor/WorldPartitionEditor/Public/WorldPartitionEditorModule.h
+++ b/Engine/Source/Editor/WorldPartitionEditor/Public/WorldPartitionEditorModule.h
@@ -54,7 +54,7 @@ public:
 	virtual void SetDisableLoadingInEditor(bool bInDisableLoadingInEditor) override;
 
 	/**
-	 * Returns if pie is disabled or not..
+	 * Returns if pie is disabled or not.
 	 */
 	virtual bool GetDisablePIE() const override;
 	virtual void SetDisablePIE(bool bInDisablePIE) override;
@@ -65,6 +65,12 @@ public:
 	virtual bool GetDisableBugIt() const override;
 	virtual void SetDisableBugIt(bool bInDisableBugIt) override;
 
+	/**
+	 * Returns if world partition is in advanced mode or not.
+	 */
+	virtual bool GetAdvancedMode() const override;
+	virtual void SetAdvancedMode(bool bInAdvancedMode) override;
+
 	/**
 	 * Convert the specified map to a world partition map.
 	 */
diff --git a/Engine/Source/Runtime/Engine/Private/WorldPartition/WorldPartition.cpp b/Engine/Source/Runtime/Engine/Private/WorldPartition/WorldPartition.cpp
--- a/Engine/Source/Runtime/Engine/Private/WorldPartition/WorldPartition.cpp
+++ b/Engine/Source/Runtime/Engine/Private/WorldPartition/WorldPartition.cpp
@@ -1189,9 +1189,6 @@ bool UWorldPartition::CanBeUsedByLevelInstance() const
 
 void UWorldPartition::SetCanBeUsedByLevelInstance(bool bInCanBeUsedByLevelInstance)
 {
-	// Only allowed to be called when WorldPartition is not initialized
-	check(!IsInitialized());
-
 	if (bCanBeUsedByLevelInstance != bInCanBeUsedByLevelInstance)
 	{
 		const FScopedTransaction Transaction(LOCTEXT("EditorWorldPartitionCanBeUsedByLevelInstance", "Set WorldPartition CanBeUsedByLevelInstance"));
diff --git a/Engine/Source/Runtime/Engine/Public/WorldPartition/IWorldPartitionEditorModule.h b/Engine/Source/Runtime/Engine/Public/WorldPartition/IWorldPartitionEditorModule.h
--- a/Engine/Source/Runtime/Engine/Public/WorldPartition/IWorldPartitionEditorModule.h
+++ b/Engine/Source/Runtime/Engine/Public/WorldPartition/IWorldPartitionEditorModule.h
@@ -45,6 +45,9 @@ public:
 	virtual bool GetDisableBugIt() const = 0;
 	virtual void SetDisableBugIt(bool bInDisableBugIt) = 0;
 
+	virtual bool GetAdvancedMode() const = 0;
+	virtual void SetAdvancedMode(bool bInAdvancedMode) = 0;
+
 	virtual bool IsEditingContentBundle() const = 0;
 	virtual bool IsEditingContentBundle(const FGuid& ContentBundleGuid) const = 0;
 
diff --git a/Engine/Source/Runtime/Engine/Public/WorldPartition/WorldPartition.h b/Engine/Source/Runtime/Engine/Public/WorldPartition/WorldPartition.h
--- a/Engine/Source/Runtime/Engine/Public/WorldPartition/WorldPartition.h
+++ b/Engine/Source/Runtime/Engine/Public/WorldPartition/WorldPartition.h
@@ -308,7 +308,7 @@ private:
 	bool bShouldCheckEnableStreamingWarning;
 
 	/** Whether Level Instance can reference this partition. */
-	UPROPERTY(EditAnywhere, Category = "WorldPartitionSetup", AdvancedDisplay, meta = (EditConditionHides, HideEditConditionToggle, EditCondition = "!bEnableStreaming"))
+	UPROPERTY()
 	bool bCanBeUsedByLevelInstance;
 #endif
 
