World Partition: add a way to override the online documentation URL.

diff --git a/Engine/Source/Editor/WorldPartitionEditor/Private/WorldPartition/Customizations/WorldPartitionDetailsCustomization.cpp b/Engine/Source/Editor/WorldPartitionEditor/Private/WorldPartition/Customizations/WorldPartitionDetailsCustomization.cpp
--- a/Engine/Source/Editor/WorldPartitionEditor/Private/WorldPartition/Customizations/WorldPartitionDetailsCustomization.cpp
+++ b/Engine/Source/Editor/WorldPartitionEditor/Private/WorldPartition/Customizations/WorldPartitionDetailsCustomization.cpp
@@ -181,7 +181,14 @@ void FWorldPartitionDetails::HandleWorldPartitionEnableStreamingChanged(ECheckBo
 
 			if (FMessageDialog::Open(EAppMsgType::YesNo, LOCTEXT("WorldPartitionEnableStreamingDialolg", "Please refer to our documentation for how to set up streaming.\n\nWould you like to open it now? ")) == EAppReturnType::Yes)
 			{
-				IDocumentation::Get()->Open(TEXT("world-partition-in-unreal-engine"), FDocumentationSourceInfo(TEXT("worldpartition")));
+				if (GetDefault<UWorldPartitionEditorSettings>()->DocURLOverride.IsEmpty())
+				{
+					IDocumentation::Get()->Open(TEXT("world-partition-in-unreal-engine"), FDocumentationSourceInfo(TEXT("worldpartition")));
+				}
+				else
+				{
+					FPlatformProcess::LaunchURL(*GetDefault<UWorldPartitionEditorSettings>()->DocURLOverride, nullptr, nullptr);
+				}
 			}
 
 			WorldPartition->bStreamingWasEnabled = true;
diff --git a/Engine/Source/Editor/WorldPartitionEditor/Private/WorldPartitionEditorModule.cpp b/Engine/Source/Editor/WorldPartitionEditor/Private/WorldPartitionEditorModule.cpp
--- a/Engine/Source/Editor/WorldPartitionEditor/Private/WorldPartitionEditorModule.cpp
+++ b/Engine/Source/Editor/WorldPartitionEditor/Private/WorldPartitionEditorModule.cpp
@@ -291,6 +291,16 @@ void FWorldPartitionEditorModule::SetAdvancedMode(bool bInAdvancedMode)
 	GetMutableDefault<UWorldPartitionEditorSettings>()->bAdvancedMode = bInAdvancedMode;
 }
 
+FString FWorldPartitionEditorModule::GetDocURLOverride() const
+{
+	return GetDefault<UWorldPartitionEditorSettings>()->DocURLOverride;
+}
+
+void FWorldPartitionEditorModule::SetDocURLOverride(const FString& InDocURLOverride)
+{
+	GetMutableDefault<UWorldPartitionEditorSettings>()->DocURLOverride = InDocURLOverride;
+}
+
 void FWorldPartitionEditorModule::OnConvertMap()
 {
 	IContentBrowserSingleton& ContentBrowserSingleton = FModuleManager::LoadModuleChecked<FContentBrowserModule>("ContentBrowser").Get();
diff --git a/Engine/Source/Editor/WorldPartitionEditor/Public/WorldPartition/WorldPartitionEditorSettings.h b/Engine/Source/Editor/WorldPartitionEditor/Public/WorldPartition/WorldPartitionEditorSettings.h
--- a/Engine/Source/Editor/WorldPartitionEditor/Public/WorldPartition/WorldPartitionEditorSettings.h
+++ b/Engine/Source/Editor/WorldPartitionEditor/Public/WorldPartition/WorldPartitionEditorSettings.h
@@ -28,4 +28,5 @@ public:
 	bool bDisableBugIt;
 	bool bDisablePIE;
 	bool bAdvancedMode;
+	FString DocURLOverride;
 };
\ No newline at end of file
diff --git a/Engine/Source/Editor/WorldPartitionEditor/Public/WorldPartitionEditorModule.h b/Engine/Source/Editor/WorldPartitionEditor/Public/WorldPartitionEditorModule.h
--- a/Engine/Source/Editor/WorldPartitionEditor/Public/WorldPartitionEditorModule.h
+++ b/Engine/Source/Editor/WorldPartitionEditor/Public/WorldPartitionEditorModule.h
@@ -71,6 +71,12 @@ public:
 	virtual bool GetAdvancedMode() const override;
 	virtual void SetAdvancedMode(bool bInAdvancedMode) override;
 
+	/**
+	 * Returns the documentation URL override.
+	 */
+	virtual FString GetDocURLOverride() const override;
+	virtual void SetDocURLOverride(const FString& InDocURLOverride) override;
+
 	/**
 	 * Convert the specified map to a world partition map.
 	 */
diff --git a/Engine/Source/Runtime/Engine/Public/WorldPartition/IWorldPartitionEditorModule.h b/Engine/Source/Runtime/Engine/Public/WorldPartition/IWorldPartitionEditorModule.h
--- a/Engine/Source/Runtime/Engine/Public/WorldPartition/IWorldPartitionEditorModule.h
+++ b/Engine/Source/Runtime/Engine/Public/WorldPartition/IWorldPartitionEditorModule.h
@@ -48,6 +48,9 @@ public:
 	virtual bool GetAdvancedMode() const = 0;
 	virtual void SetAdvancedMode(bool bInAdvancedMode) = 0;
 
+	virtual FString GetDocURLOverride() const = 0;
+	virtual void SetDocURLOverride(const FString& InDocURLOverride) = 0;
+
 	virtual bool IsEditingContentBundle() const = 0;
 	virtual bool IsEditingContentBundle(const FGuid& ContentBundleGuid) const = 0;
 
