Fix crash when deleting a level that has references to sequencer assets.
Make sure that the world pointer is not exposed so it won't get nulled out by reference replacement when deleting a world.

diff --git a/Engine/Source/Editor/UnrealEd/Private/WorldPartition/WorldPartitionNavigationDataBuilder.cpp b/Engine/Source/Editor/UnrealEd/Private/WorldPartition/WorldPartitionNavigationDataBuilder.cpp
--- a/Engine/Source/Editor/UnrealEd/Private/WorldPartition/WorldPartitionNavigationDataBuilder.cpp
+++ b/Engine/Source/Editor/UnrealEd/Private/WorldPartition/WorldPartitionNavigationDataBuilder.cpp
@@ -286,7 +286,7 @@ bool UWorldPartitionNavigationDataBuilder::GenerateNavigationData(UWorldPartitio
 	CallCount++;
 	UE_LOG(LogWorldPartitionNavigationDataBuilder, Log, TEXT("Iteration %i. GenerateNavigationData for LoadedBounds %s"), CallCount, *LoadedBounds.ToString());
 
-	UWorld* World = WorldPartition->World;
+	UWorld* World = WorldPartition->GetWorld();
 
 	// Generate navmesh
 	// Make sure navigation is added and initialized in EditorWorldPartitionBuildMode
diff --git a/Engine/Source/Editor/WorldPartitionEditor/Private/WorldPartition/SWorldPartitionEditor.cpp b/Engine/Source/Editor/WorldPartitionEditor/Private/WorldPartition/SWorldPartitionEditor.cpp
--- a/Engine/Source/Editor/WorldPartitionEditor/Private/WorldPartition/SWorldPartitionEditor.cpp
+++ b/Engine/Source/Editor/WorldPartitionEditor/Private/WorldPartition/SWorldPartitionEditor.cpp
@@ -40,7 +40,7 @@ SWorldPartitionEditor::~SWorldPartitionEditor()
 	{
 		if (UWorldPartition* WorldPartition = World->GetWorldPartition())
 		{
-			check(WorldPartition->World == World);			
+			check(WorldPartition->GetWorld() == World);
 			if (WorldPartition->WorldPartitionEditor)
 			{
 				check(WorldPartition->WorldPartitionEditor == this);
@@ -78,7 +78,7 @@ TSharedRef<SWidget> SWorldPartitionEditor::ConstructContentWidget()
 	{
 		if (UWorldPartition* WorldPartition = World->GetWorldPartition())
 		{
-			check(WorldPartition->World == World);
+			check(WorldPartition->GetWorld() == World);
 			EditorName = WorldPartition->GetWorldPartitionEditorName();
 			WorldPartition->WorldPartitionEditor = this;
 		}
diff --git a/Engine/Source/Runtime/Engine/Public/WorldPartition/WorldPartition.h b/Engine/Source/Runtime/Engine/Public/WorldPartition/WorldPartition.h
--- a/Engine/Source/Runtime/Engine/Public/WorldPartition/WorldPartition.h
+++ b/Engine/Source/Runtime/Engine/Public/WorldPartition/WorldPartition.h
@@ -327,14 +327,13 @@ public:
 	UPROPERTY()
 	TObjectPtr<UWorldPartitionRuntimeHash> RuntimeHash;
 
-	UPROPERTY(Transient)
-	TObjectPtr<UWorld> World;
-
 	/** Enables streaming for this world. */
 	UPROPERTY()
 	bool bEnableStreaming;
 
 private:
+	TObjectPtr<UWorld> World;
+
 #if WITH_EDITOR
 	bool bForceGarbageCollection;
 	bool bForceGarbageCollectionPurge;
