Fix Non-unity build

#rnx

diff --git a/Engine/Source/Editor/LevelInstanceEditor/Private/LevelInstanceFilterPropertyTypeCustomization.cpp b/Engine/Source/Editor/LevelInstanceEditor/Private/LevelInstanceFilterPropertyTypeCustomization.cpp
--- a/Engine/Source/Editor/LevelInstanceEditor/Private/LevelInstanceFilterPropertyTypeCustomization.cpp
+++ b/Engine/Source/Editor/LevelInstanceEditor/Private/LevelInstanceFilterPropertyTypeCustomization.cpp
@@ -7,6 +7,7 @@
 #include "WorldPartition/Filter/WorldPartitionActorFilter.h"
 #include "Engine/World.h"
 #include "GameFramework/Actor.h"
+#include "ScopedTransaction.h"
 
 #define LOCTEXT_NAMESPACE "WorldPartitionActorFilter"
 
diff --git a/Engine/Source/Editor/WorldPartitionEditor/Private/WorldPartition/Filter/WorldPartitionActorFilterPropertyTypeCustomization.cpp b/Engine/Source/Editor/WorldPartitionEditor/Private/WorldPartition/Filter/WorldPartitionActorFilterPropertyTypeCustomization.cpp
--- a/Engine/Source/Editor/WorldPartitionEditor/Private/WorldPartition/Filter/WorldPartitionActorFilterPropertyTypeCustomization.cpp
+++ b/Engine/Source/Editor/WorldPartitionEditor/Private/WorldPartition/Filter/WorldPartitionActorFilterPropertyTypeCustomization.cpp
@@ -7,6 +7,8 @@
 #include "Containers/Array.h"
 #include "CoreTypes.h"
 #include "DetailWidgetRow.h"
+#include "Widgets/Input/SComboButton.h"
+#include "Widgets/Input/SButton.h"
 #include "Internationalization/Internationalization.h"
 #include "LevelInstance/LevelInstanceSubsystem.h"
 #include "LevelInstance/LevelInstanceInterface.h"
diff --git a/Engine/Source/Editor/WorldPartitionEditor/Private/WorldPartition/Filter/WorldPartitionActorFilterTreeItems.cpp b/Engine/Source/Editor/WorldPartitionEditor/Private/WorldPartition/Filter/WorldPartitionActorFilterTreeItems.cpp
--- a/Engine/Source/Editor/WorldPartitionEditor/Private/WorldPartition/Filter/WorldPartitionActorFilterTreeItems.cpp
+++ b/Engine/Source/Editor/WorldPartitionEditor/Private/WorldPartition/Filter/WorldPartitionActorFilterTreeItems.cpp
@@ -6,8 +6,14 @@
 #include "WorldPartition/DataLayer/DataLayerAsset.h"
 #include "Widgets/Text/STextBlock.h"
 #include "Widgets/SBoxPanel.h"
+#include "Widgets/Images/SImage.h"
+#include "Widgets/SWidget.h"
+#include "Widgets/SCompoundWidget.h"
+#include "Widgets/Views/STableRow.h"
+#include "Widgets/Views/SListView.h"
 #include "ISceneOutliner.h"
 #include "SceneOutlinerPublicTypes.h"
+#include "SceneOutlinerStandaloneTypes.h"
 #include "Styling/SlateIconFinder.h"
 
 #define LOCTEXT_NAMESPACE "WorldPartitionActorFilter"
