Convert raw pointers to weak pointers in world partition editor.

diff --git a/Engine/Source/Editor/WorldPartitionEditor/Private/WorldPartition/SWorldPartitionEditorGrid.cpp b/Engine/Source/Editor/WorldPartitionEditor/Private/WorldPartition/SWorldPartitionEditorGrid.cpp
--- a/Engine/Source/Editor/WorldPartitionEditor/Private/WorldPartition/SWorldPartitionEditorGrid.cpp
+++ b/Engine/Source/Editor/WorldPartitionEditor/Private/WorldPartition/SWorldPartitionEditorGrid.cpp
@@ -17,7 +17,7 @@ TMap<FName, SWorldPartitionEditorGrid::PartitionEditorGridCreateInstanceFunc> SW
 void SWorldPartitionEditorGrid::Construct(const FArguments& InArgs)
 {
 	World = InArgs._InWorld;
-	WorldPartition = World ? World->GetWorldPartition() : nullptr;
+	WorldPartition = World.IsValid() ? World->GetWorldPartition() : nullptr;
 
 	if (IsDisabled())
 	{
@@ -109,7 +109,7 @@ void SWorldPartitionEditorGrid::Refresh()
 
 bool SWorldPartitionEditorGrid::IsDisabled() const
 {
-	return !WorldPartition;
+	return !WorldPartition.IsValid();
 }
 
 #undef LOCTEXT_NAMESPACE
diff --git a/Engine/Source/Editor/WorldPartitionEditor/Private/WorldPartition/SWorldPartitionEditorGrid2D.cpp b/Engine/Source/Editor/WorldPartitionEditor/Private/WorldPartition/SWorldPartitionEditorGrid2D.cpp
--- a/Engine/Source/Editor/WorldPartitionEditor/Private/WorldPartition/SWorldPartitionEditorGrid2D.cpp
+++ b/Engine/Source/Editor/WorldPartitionEditor/Private/WorldPartition/SWorldPartitionEditorGrid2D.cpp
@@ -287,7 +287,7 @@ void SWorldPartitionEditorGrid2D::Construct(const FArguments& InArgs)
 {
 	SWorldPartitionEditorGrid::Construct(SWorldPartitionEditorGrid::FArguments().InWorld(InArgs._InWorld));
 
-	if (WorldPartition)
+	if (GetWorldPartition())
 	{
 		UpdateWorldMiniMapDetails();
 		bShowActors = !WorldMiniMapBrush.HasUObject();
@@ -377,7 +377,7 @@ void SWorldPartitionEditorGrid2D::Construct(const FArguments& InArgs)
 			ToolbarBuilder.AddWidget(SNew(SCheckBox)
 				.ForegroundColor(FSlateColor::UseForeground())
 				.IsChecked(GetMutableDefault<UWorldPartitionEditorPerProjectUserSettings>()->GetShowCellCoords() ? ECheckBoxState::Checked : ECheckBoxState::Unchecked)
-				.Visibility_Lambda([this]() { return WorldPartition->IsStreamingEnabled() ? EVisibility::Visible : EVisibility::Hidden; })
+				.Visibility_Lambda([this]() { return (GetWorldPartition() && GetWorldPartition()->IsStreamingEnabled()) ? EVisibility::Visible : EVisibility::Hidden; })
 				.OnCheckStateChanged(FOnCheckStateChanged::CreateLambda([=](ECheckBoxState State) { GetMutableDefault<UWorldPartitionEditorPerProjectUserSettings>()->SetShowCellCoords(State == ECheckBoxState::Checked); }))
 				[
 					SNew(SBox)
@@ -394,8 +394,11 @@ void SWorldPartitionEditorGrid2D::Construct(const FArguments& InArgs)
 				.ForegroundColor(FSlateColor::UseForeground())
 				.Text_Lambda([this]()
 				{
-					UWorldPartitionSubsystem* WorldPartitionSubsystem = UWorld::GetSubsystem<UWorldPartitionSubsystem>(World);
-					return (GEditor->GetSelectedActors()->Num() || WorldPartitionSubsystem->SelectedActorDescs.Num()) ? LOCTEXT("FocusSelection", "Focus Selection") : LOCTEXT("FocusWorld", "Focus World");
+					if (UWorldPartitionSubsystem* WorldPartitionSubsystem = UWorld::GetSubsystem<UWorldPartitionSubsystem>(GetWorld()))
+					{
+						return (GEditor->GetSelectedActors()->Num() || WorldPartitionSubsystem->SelectedActorDescs.Num()) ? LOCTEXT("FocusSelection", "Focus Selection") : LOCTEXT("FocusWorld", "Focus World");
+					}
+					return LOCTEXT("FocusWorld", "Focus World");
 				})
 				.OnClicked(this, &SWorldPartitionEditorGrid2D::FocusSelection)
 				.IsEnabled_Lambda([this]() { return IsInteractive(); }));
@@ -404,7 +407,7 @@ void SWorldPartitionEditorGrid2D::Construct(const FArguments& InArgs)
 				.ForegroundColor(FSlateColor::UseForeground())
 				.Text(LOCTEXT("FocusLoadedRegions", "Focus Loaded Regions"))
 				.OnClicked(this, &SWorldPartitionEditorGrid2D::FocusLoadedRegions)
-				.IsEnabled_Lambda([this]() { return IsInteractive() && WorldPartition && WorldPartition->HasLoadedUserCreatedRegions(); })
+				.IsEnabled_Lambda([this]() { return IsInteractive() && GetWorldPartition() && GetWorldPartition()->HasLoadedUserCreatedRegions(); })
 				.Visibility_Lambda([this]() { return GetDefault<UWorldPartitionEditorSettings>()->bDisableLoadingInEditor ? EVisibility::Hidden : EVisibility::Visible; }));
 		}
 
@@ -429,7 +432,7 @@ void SWorldPartitionEditorGrid2D::Construct(const FArguments& InArgs)
 		[
 			SAssignNew(ViewportWidget, SWorldPartitionViewportWidget)
 			.Clickable(false)
-			.Visibility_Lambda([this]() { return ViewportWidget->GetVisibility(World); })
+			.Visibility_Lambda([this]() { return ViewportWidget->GetVisibility(GetWorld()); })
 		]
 	];
 
@@ -487,7 +490,7 @@ void SWorldPartitionEditorGrid2D::Construct(const FArguments& InArgs)
 
 void SWorldPartitionEditorGrid2D::UpdateWorldMiniMapDetails()
 {
-	if (AWorldPartitionMiniMap* WorldMiniMap = FWorldPartitionMiniMapHelper::GetWorldPartitionMiniMap(World))
+	if (AWorldPartitionMiniMap* WorldMiniMap = FWorldPartitionMiniMapHelper::GetWorldPartitionMiniMap(GetWorld()))
 	{
 		WorldMiniMapBounds = FBox2D(FVector2D(WorldMiniMap->MiniMapWorldBounds.Min), FVector2D(WorldMiniMap->MiniMapWorldBounds.Max));
 		if (UTexture2D* MiniMapTexture = WorldMiniMap->MiniMapTexture)
@@ -502,7 +505,7 @@ void SWorldPartitionEditorGrid2D::UpdateWorldMiniMapDetails()
 void SWorldPartitionEditorGrid2D::CreateRegionFromSelection()
 {
 	const FBox RegionBox(FVector(SelectBoxGridSnapped.Min.X, SelectBoxGridSnapped.Min.Y, -HALF_WORLD_MAX), FVector(SelectBoxGridSnapped.Max.X, SelectBoxGridSnapped.Max.Y, HALF_WORLD_MAX));
-	UWorldPartitionEditorLoaderAdapter* EditorLoaderAdapter = WorldPartition->CreateEditorLoaderAdapter<FLoaderAdapterShape>(World, RegionBox, TEXT("Loaded Region"));
+	UWorldPartitionEditorLoaderAdapter* EditorLoaderAdapter = GetWorldPartition()->CreateEditorLoaderAdapter<FLoaderAdapterShape>(GetWorld(), RegionBox, TEXT("Loaded Region"));
 	EditorLoaderAdapter->GetLoaderAdapter()->SetUserCreated(true);
 	EditorLoaderAdapter->GetLoaderAdapter()->Load();
 	
@@ -542,7 +545,7 @@ void SWorldPartitionEditorGrid2D::UnloadSelectedRegions()
 			if (UWorldPartitionEditorLoaderAdapter* EditorLoaderAdapter = Cast<UWorldPartitionEditorLoaderAdapter>(LoaderInterface))
 			{
 				SelectedLoaderInterfaces.Remove(EditorLoaderAdapter);
-				WorldPartition->ReleaseEditorLoaderAdapter(EditorLoaderAdapter);
+				GetWorldPartition()->ReleaseEditorLoaderAdapter(EditorLoaderAdapter);
 			}
 		}
 	};
@@ -559,7 +562,7 @@ void SWorldPartitionEditorGrid2D::ConvertSelectedRegionsToActors()
 	FLoaderInterfaceSet TmpSelectedLoaderInterfaces = MoveTemp(SelectedLoaderInterfaces);
 	ClearSelection();
 
-	const FBox WorldBounds = WorldPartition->GetRuntimeWorldBounds();
+	const FBox WorldBounds = GetWorldPartition()->GetRuntimeWorldBounds();
 	for (const TWeakInterfacePtr<IWorldPartitionActorLoaderInterface>& SelectedLoaderAdapter : TmpSelectedLoaderInterfaces)
 	{
 		if (UWorldPartitionEditorLoaderAdapter* EditorLoaderAdapter = Cast<UWorldPartitionEditorLoaderAdapter>(SelectedLoaderAdapter.Get()))
@@ -569,7 +572,7 @@ void SWorldPartitionEditorGrid2D::ConvertSelectedRegionsToActors()
 			const FBox LoaderVolumeBox(*LoaderAdapter->GetBoundingBox());
 			const FBox ActorVolumeBox(FVector(LoaderVolumeBox.Min.X, LoaderVolumeBox.Min.Y, WorldBounds.Min.Z), FVector(LoaderVolumeBox.Max.X, LoaderVolumeBox.Max.Y, WorldBounds.Max.Z));
 
-			ALocationVolume* LocationVolume = World->SpawnActor<ALocationVolume>(ActorVolumeBox.GetCenter(), FRotator::ZeroRotator);
+			ALocationVolume* LocationVolume = GetWorld()->SpawnActor<ALocationVolume>(ActorVolumeBox.GetCenter(), FRotator::ZeroRotator);
 
 			UCubeBuilder* Builder = NewObject<UCubeBuilder>();
 			Builder->X = 1.0f;
@@ -583,7 +586,7 @@ void SWorldPartitionEditorGrid2D::ConvertSelectedRegionsToActors()
 
 			LoaderAdapter->Unload();
 			
-			WorldPartition->ReleaseEditorLoaderAdapter(EditorLoaderAdapter);
+			GetWorldPartition()->ReleaseEditorLoaderAdapter(EditorLoaderAdapter);
 		}
 		else
 		{
@@ -600,7 +603,7 @@ void SWorldPartitionEditorGrid2D::MoveCameraHere()
 	FVector WorldLocation = FVector(MouseCursorPosWorld, 0);
 
 	FHitResult HitResult;
-	const bool bHitResultValid = HitTestZFromLocation(World, MouseCursorPosWorld, HitResult);
+	const bool bHitResultValid = HitTestZFromLocation(GetWorld(), MouseCursorPosWorld, HitResult);
 
 	for (FLevelEditorViewportClient* LevelVC : GEditor->GetLevelViewportClients())
 	{
@@ -619,7 +622,7 @@ void SWorldPartitionEditorGrid2D::PlayFromHere()
 	FRotator StartRotation;
 
 	FHitResult HitResult;
-	const bool bHitResultValid = HitTestZFromLocation(World, MouseCursorPosWorld, HitResult);
+	const bool bHitResultValid = HitTestZFromLocation(GetWorld(), MouseCursorPosWorld, HitResult);
 
 	FLevelEditorModule& LevelEditorModule = FModuleManager::GetModuleChecked<FLevelEditorModule>("LevelEditor");
 	TSharedPtr<IAssetViewport> ActiveLevelViewport = LevelEditorModule.GetFirstActiveViewport();
@@ -656,7 +659,7 @@ void SWorldPartitionEditorGrid2D::LoadFromHere()
 	}
 
 	// Load box
-	UWorldPartitionEditorLoaderAdapter* EditorLoaderAdapter = WorldPartition->CreateEditorLoaderAdapter<FLoaderAdapterShape>(World, LoadCellsBox, TEXT("Loaded Region"));
+	UWorldPartitionEditorLoaderAdapter* EditorLoaderAdapter = GetWorldPartition()->CreateEditorLoaderAdapter<FLoaderAdapterShape>(GetWorld(), LoadCellsBox, TEXT("Loaded Region"));
 	EditorLoaderAdapter->GetLoaderAdapter()->Load();
 }
 
@@ -665,7 +668,7 @@ void SWorldPartitionEditorGrid2D::BugItHere()
 	FVector WorldLocation = FVector(MouseCursorPosWorld, 0);
 
 	FHitResult HitResult;
-	const bool bHitResultValid = HitTestZFromLocation(World, MouseCursorPosWorld, HitResult);
+	const bool bHitResultValid = HitTestZFromLocation(GetWorld(), MouseCursorPosWorld, HitResult);
 
 	if (bHitResultValid && GCurrentLevelEditingViewportClient)
 	{
@@ -804,7 +807,7 @@ FReply SWorldPartitionEditorGrid2D::OnMouseButtonUp(const FGeometry& MyGeometry,
 			FLoaderInterfaceSet LoaderAdaptersToSelect;
 			if (bIsDragSelecting)
 			{
-				ForEachIntersectingLoaderAdapters(WorldPartition, SelectBoxGridSnapped, [&LoaderAdaptersToSelect](UObject* AdapterObject)
+				ForEachIntersectingLoaderAdapters(GetWorldPartition(), SelectBoxGridSnapped, [&LoaderAdaptersToSelect](UObject* AdapterObject)
 				{
 					LoaderAdaptersToSelect.Add(AdapterObject);
 					return true;
@@ -1030,7 +1033,7 @@ void SWorldPartitionEditorGrid2D::Tick(const FGeometry& AllottedGeometry, const
 	DirtyActorGuids.Reset();
 	ShownLoaderInterfaces.Reset(); 
 	
-	for (UWorldPartitionEditorLoaderAdapter* EditorLoaderAdapter : WorldPartition->GetRegisteredEditorLoaderAdapters())
+	for (UWorldPartitionEditorLoaderAdapter* EditorLoaderAdapter : GetWorldPartition()->GetRegisteredEditorLoaderAdapters())
 	{
 		check(EditorLoaderAdapter);
 		IWorldPartitionActorLoaderInterface::ILoaderAdapter* LoaderAdapter = EditorLoaderAdapter->GetLoaderAdapter();
@@ -1045,11 +1048,11 @@ void SWorldPartitionEditorGrid2D::Tick(const FGeometry& AllottedGeometry, const
 	UWorldPartitionEditorHash::FForEachIntersectingActorParams ForEachIntersectingActorParams;
 	ForEachIntersectingActorParams.MinimumBox = FBox(FVector::ZeroVector, FVector(GetSelectionSnap()));
 
-	WorldPartition->EditorHash->ForEachIntersectingActor(ViewRectWorld, [&](FWorldPartitionActorDesc* ActorDesc)
+	GetWorldPartition()->EditorHash->ForEachIntersectingActor(ViewRectWorld, [&](FWorldPartitionActorDesc* ActorDesc)
 	{
 		if (bShowActors)
 		{
-			if (!WorldPartition->IsStreamingEnabled() || ActorDesc->GetIsSpatiallyLoaded())
+			if (!GetWorldPartition()->IsStreamingEnabled() || ActorDesc->GetIsSpatiallyLoaded())
 			{
 				ShownActorGuids.Add(ActorDesc->GetGuid());
 			}
@@ -1069,9 +1072,9 @@ void SWorldPartitionEditorGrid2D::Tick(const FGeometry& AllottedGeometry, const
 
 	// Also include transient actor loader adapters that might have been spawned by blutilities, etc. Since these actors can't be saved because they are transient,
 	// they will never get an actor descriptor so they will never appear in the world partition editor. Also include unsaved, newly created actors for convenience.
-	for (auto& [Reference, Actor] : WorldPartition->GetDirtyActors())
+	for (auto& [Reference, Actor] : GetWorldPartition()->GetDirtyActors())
 	{
-		if (!WorldPartition->GetActorDesc(Actor->GetActorGuid()) && Actor->Implements<UWorldPartitionActorLoaderInterface>())
+		if (!GetWorldPartition()->GetActorDesc(Actor->GetActorGuid()) && Actor->Implements<UWorldPartitionActorLoaderInterface>())
 		{
 			if (IWorldPartitionActorLoaderInterface::ILoaderAdapter* LoaderAdapter = Cast<IWorldPartitionActorLoaderInterface>(Actor)->GetLoaderAdapter())
 			{
@@ -1152,7 +1155,7 @@ void SWorldPartitionEditorGrid2D::Tick(const FGeometry& AllottedGeometry, const
 	}
 
 	// Include selected actor descriptors
-	UWorldPartitionSubsystem* WorldPartitionSubsystem = UWorld::GetSubsystem<UWorldPartitionSubsystem>(World);
+	UWorldPartitionSubsystem* WorldPartitionSubsystem = UWorld::GetSubsystem<UWorldPartitionSubsystem>(GetWorld());
 	for (FWorldPartitionActorDesc* SelectedActorDesc : WorldPartitionSubsystem->SelectedActorDescs)
 	{
 		ShownActorGuids.Add(SelectedActorDesc->GetGuid());
@@ -1164,7 +1167,7 @@ uint32 SWorldPartitionEditorGrid2D::PaintActors(const FGeometry& AllottedGeometr
 	const FBox2D ViewRect(FVector2D(ForceInitToZero), AllottedGeometry.GetLocalSize());
 	const FBox2D WorldViewRect(ScreenToWorld.TransformPoint(ViewRect.Min), ScreenToWorld.TransformPoint(ViewRect.Max));
 	const FBox ViewRectWorld(FVector(WorldViewRect.Min.X, WorldViewRect.Min.Y, -HALF_WORLD_MAX), FVector(WorldViewRect.Max.X, WorldViewRect.Max.Y, HALF_WORLD_MAX));
-	const UWorldPartitionSubsystem* WorldPartitionSubsystem = UWorld::GetSubsystem<UWorldPartitionSubsystem>(World);
+	const UWorldPartitionSubsystem* WorldPartitionSubsystem = UWorld::GetSubsystem<UWorldPartitionSubsystem>(GetWorld());
 
 	TArray<FWorldPartitionActorDescViewBoundsProxy> ActorDescList;
 	ActorDescList.Reserve(ShownActorGuids.Num());
@@ -1174,10 +1177,11 @@ uint32 SWorldPartitionEditorGrid2D::PaintActors(const FGeometry& AllottedGeometr
 	{
 		LocalHoveredLoaderAdapter = HoveredLoaderInterface->GetLoaderAdapter();
 	}
-		
+
+	const UWorldPartition* ThisWorldPartition = GetWorldPartition();
 	for (const FGuid& ActorGuid : ShownActorGuids)
 	{
-		if (FWorldPartitionActorDesc* ActorDesc = WorldPartition->GetActorDesc(ActorGuid))
+		if (const FWorldPartitionActorDesc* ActorDesc = ThisWorldPartition->GetActorDesc(ActorGuid))
 		{
 			ActorDescList.Emplace(ActorDesc, DirtyActorGuids.Contains(ActorGuid));
 		}
@@ -1502,7 +1506,7 @@ uint32 SWorldPartitionEditorGrid2D::PaintTextInfo(const FGeometry& AllottedGeome
 		FLinearColor::White);
 
 	// Show world bounds
-	const FBox WorldBounds = WorldPartition->GetRuntimeWorldBounds();
+	const FBox WorldBounds = GetWorldPartition()->GetRuntimeWorldBounds();
 	const FVector WorldBoundsExtentInKM = (WorldBounds.GetExtent() * 2.0f) / 100000.0f;
 	RulerText = FString::Printf(TEXT("%.2fx%.2fx%.2f km"), WorldBoundsExtentInKM.X, WorldBoundsExtentInKM.Y, WorldBoundsExtentInKM.Z);
 
@@ -1698,7 +1702,7 @@ int32 SWorldPartitionEditorGrid2D::OnPaint(const FPaintArgs& Args, const FGeomet
 {
 	FWeightedMovingAverageScope ProfileMeanValue(PaintTime);
 
-	if (WorldPartition)
+	if (GetWorldPartition())
 	{
 		const bool bResetView = !ScreenRect.bIsValid;
 
@@ -1776,7 +1780,7 @@ int32 SWorldPartitionEditorGrid2D::PaintMinimap(const FGeometry& AllottedGeometr
 			FLinearColor::White
 		);
 
-		if (WorldPartition->IsStreamingEnabled())
+		if (GetWorldPartition()->IsStreamingEnabled())
 		{
 			F2DRectBooleanSubtract ShadowAreas(MinimapBounds);
 			for (const FLoaderInterface& LoaderInterface : ShownLoaderInterfaces)
@@ -1966,7 +1970,7 @@ FReply SWorldPartitionEditorGrid2D::FocusSelection()
 		}
 	}
 
-	UWorldPartitionSubsystem* WorldPartitionSubsystem = UWorld::GetSubsystem<UWorldPartitionSubsystem>(World);
+	UWorldPartitionSubsystem* WorldPartitionSubsystem = UWorld::GetSubsystem<UWorldPartitionSubsystem>(GetWorld());
 	for (FWorldPartitionActorDesc* SelectedActorDesc : WorldPartitionSubsystem->SelectedActorDescs)
 	{
 		SelectionBox += SelectedActorDesc->GetEditorBounds();
@@ -1975,7 +1979,7 @@ FReply SWorldPartitionEditorGrid2D::FocusSelection()
 	if (!SelectionBox.IsValid)
 	{
 		// Override the minimap bounds if world partition minimap volumes exists
-		for (TActorIterator<AWorldPartitionMiniMapVolume> It(World); It; ++It)
+		for (TActorIterator<AWorldPartitionMiniMapVolume> It(GetWorld()); It; ++It)
 		{
 			if (AWorldPartitionMiniMapVolume* WorldPartitionMiniMapVolume = *It)
 			{
@@ -1985,7 +1989,7 @@ FReply SWorldPartitionEditorGrid2D::FocusSelection()
 
 		if (!SelectionBox.IsValid)
 		{
-			SelectionBox = WorldPartition->GetRuntimeWorldBounds();
+			SelectionBox = GetWorldPartition()->GetRuntimeWorldBounds();
 		}
 	}
 
@@ -1997,7 +2001,7 @@ FReply SWorldPartitionEditorGrid2D::FocusLoadedRegions()
 {
 	FBox SelectionBox(ForceInit);
 
-	for (UWorldPartitionEditorLoaderAdapter* EditorLoaderAdapter : WorldPartition->GetRegisteredEditorLoaderAdapters())
+	for (UWorldPartitionEditorLoaderAdapter* EditorLoaderAdapter : GetWorldPartition()->GetRegisteredEditorLoaderAdapters())
 	{
 		check(EditorLoaderAdapter);
 		IWorldPartitionActorLoaderInterface::ILoaderAdapter* LoaderAdapter = EditorLoaderAdapter->GetLoaderAdapter();
@@ -2009,7 +2013,7 @@ FReply SWorldPartitionEditorGrid2D::FocusLoadedRegions()
 		}
 	}
 
-	for (TActorIterator<AActor> ActorIt(World); ActorIt; ++ActorIt)
+	for (TActorIterator<AActor> ActorIt(GetWorld()); ActorIt; ++ActorIt)
 	{
 		if (AActor* Actor = *ActorIt)
 		{
@@ -2076,7 +2080,7 @@ void SWorldPartitionEditorGrid2D::UpdateTransform() const
 
 void SWorldPartitionEditorGrid2D::UpdateSelectionBox(bool bSnap)
 {
-	if (WorldPartition->IsStreamingEnabled())
+	if (GetWorldPartition()->IsStreamingEnabled())
 	{
 		const FBox2D SelectBox2D(FVector2D::Min(SelectionStart, SelectionEnd), FVector2D::Max(SelectionStart, SelectionEnd));
 
@@ -2105,4 +2109,4 @@ void SWorldPartitionEditorGrid2D::ClearSelection()
 	SelectBoxGridSnapped.Init();
 }
 
-#undef LOCTEXT_NAMESPACE
+#undef LOCTEXT_NAMESPACE
\ No newline at end of file
diff --git a/Engine/Source/Editor/WorldPartitionEditor/Private/WorldPartition/SWorldPartitionEditorGridSpatialHash.cpp b/Engine/Source/Editor/WorldPartitionEditor/Private/WorldPartition/SWorldPartitionEditorGridSpatialHash.cpp
--- a/Engine/Source/Editor/WorldPartitionEditor/Private/WorldPartition/SWorldPartitionEditorGridSpatialHash.cpp
+++ b/Engine/Source/Editor/WorldPartitionEditor/Private/WorldPartition/SWorldPartitionEditorGridSpatialHash.cpp
@@ -28,7 +28,7 @@ int64 SWorldPartitionEditorGridSpatialHash::GetSelectionSnap() const
 
 int64 SWorldPartitionEditorGridSpatialHash::GetSelectionSnap(float& FadeRatio, float& CellScreenSize) const
 {
-	UWorldPartitionEditorSpatialHash* EditorSpatialHash = (UWorldPartitionEditorSpatialHash*)WorldPartition->EditorHash;
+	UWorldPartitionEditorSpatialHash* EditorSpatialHash = (UWorldPartitionEditorSpatialHash*)GetWorldPartition()->EditorHash;
 	
 	int64 EffectiveCellSize = EditorSpatialHash->CellSize;
 	for (;;)
@@ -50,7 +50,7 @@ int64 SWorldPartitionEditorGridSpatialHash::GetSelectionSnap(float& FadeRatio, f
 
 int32 SWorldPartitionEditorGridSpatialHash::PaintGrid(const FGeometry& AllottedGeometry, const FSlateRect& MyCullingRect, FSlateWindowElementList& OutDrawElements, int32 LayerId) const
 {
-	UWorldPartitionEditorSpatialHash* EditorSpatialHash = (UWorldPartitionEditorSpatialHash*)WorldPartition->EditorHash;
+	UWorldPartitionEditorSpatialHash* EditorSpatialHash = (UWorldPartitionEditorSpatialHash*)GetWorldPartition()->EditorHash;
 
 	// Found the best cell size depending on the current zoom
 	float FadeRatio;
diff --git a/Engine/Source/Editor/WorldPartitionEditor/Public/WorldPartition/SWorldPartitionEditorGrid.h b/Engine/Source/Editor/WorldPartitionEditor/Public/WorldPartition/SWorldPartitionEditorGrid.h
--- a/Engine/Source/Editor/WorldPartitionEditor/Public/WorldPartition/SWorldPartitionEditorGrid.h
+++ b/Engine/Source/Editor/WorldPartitionEditor/Public/WorldPartition/SWorldPartitionEditorGrid.h
@@ -66,12 +66,18 @@ public:
 	void Refresh();
 	virtual void FocusBox(const FBox& Box) const {}
 
-	UWorld* World = nullptr;
-	UWorldPartition* WorldPartition = nullptr;
+	inline UWorld* GetWorld() { return World.Get(); }
+	inline const UWorld* GetWorld() const { return World.Get(); }
+	
+	inline UWorldPartition* GetWorldPartition() { return WorldPartition.Get(); }
+	inline const UWorldPartition* GetWorldPartition() const { return WorldPartition.Get(); }
 
 protected:
 	static TMap<FName, PartitionEditorGridCreateInstanceFunc> PartitionEditorGridCreateInstanceFactory;
 
 private:
 	bool IsDisabled() const;
+
+	TWeakObjectPtr<UWorld> World;
+	TWeakObjectPtr<UWorldPartition> WorldPartition;
 };
\ No newline at end of file
