Actor Filter:
- Replace dropdown with embedded filter in details panel
- Add item tooltip

#rnx

diff --git a/Engine/Source/Editor/WorldPartitionEditor/Private/WorldPartition/Filter/WorldPartitionActorFilterPropertyTypeCustomization.cpp b/Engine/Source/Editor/WorldPartitionEditor/Private/WorldPartition/Filter/WorldPartitionActorFilterPropertyTypeCustomization.cpp
--- a/Engine/Source/Editor/WorldPartitionEditor/Private/WorldPartition/Filter/WorldPartitionActorFilterPropertyTypeCustomization.cpp
+++ b/Engine/Source/Editor/WorldPartitionEditor/Private/WorldPartition/Filter/WorldPartitionActorFilterPropertyTypeCustomization.cpp
@@ -54,7 +54,7 @@ void FWorldPartitionActorFilterPropertyTypeCustomization::CustomizeHeader(TShare
 
 	FSceneOutlinerInitializationOptions InitOptions;
 	{
-		InitOptions.bShowHeaderRow = true;
+		InitOptions.bShowHeaderRow = false;
 		InitOptions.bFocusSearchBoxWhenOpened = true;
 		InitOptions.bShowCreateNewFolder = false;
 		InitOptions.ModeFactory = ModeFactory;
@@ -96,38 +96,29 @@ void FWorldPartitionActorFilterPropertyTypeCustomization::CustomizeHeader(TShare
 	]
 	.ValueContent()
 	[
-		SNew(SHorizontalBox)
-		+SHorizontalBox::Slot()
+		SNew(SBox)
+		.MinDesiredWidth(2000.f) // Set very large as it is the only way for the Outliner to expand to take all the space
+		.MinDesiredHeight(150.f)
+		.MaxDesiredHeight(300.f)
 		[
-			SNew(SComboButton)
-			.ContentPadding(2.f)
-			.HasDownArrow(false)
-			.ButtonContent()
+			SNew(SVerticalBox)
+			+ SVerticalBox::Slot()
+			.FillHeight(1.0f)
 			[
-				SNew(STextBlock)
-				.Text(LOCTEXT("FilterText", "Select Filters..."))
-				.Font(FAppStyle::GetFontStyle(TEXT("PropertyWindow.NormalFont")))
+				SceneOutliner.ToSharedRef()
 			]
-			.MenuContent()
+			+ SVerticalBox::Slot()
+			.AutoHeight()
+			.Padding(0.0f, 4.0f, 0.0f, 4.0f)
 			[
-				SNew(SBox)
-				.MaxDesiredHeight(400.0f)
-				.WidthOverride(300.0f)
-				[
-					SceneOutliner.ToSharedRef()
-				]
+				SNew(SButton)
+				.VAlign(VAlign_Center)
+				.HAlign(HAlign_Center)
+				.OnClicked_Lambda(OnApplyFilter)
+				.Text(LOCTEXT("ApplyFilter", "Apply"))
+				.ToolTipText(LOCTEXT("ApplyFilterToolTip", "Apply Filter"))
 			]
 		]
-		+SHorizontalBox::Slot()
-		.AutoWidth()
-		[
-			SNew(SButton)
-			.VAlign(VAlign_Center)
-			.HAlign(HAlign_Center)
-			.OnClicked_Lambda(OnApplyFilter)
-			.Text(LOCTEXT("ApplyFilter", "Apply"))
-			.ToolTipText(LOCTEXT("ApplyFilterToolTip", "Apply Filter"))
-		]
 	];
 		
 	// This gets called when hitting ResetToDefault
diff --git a/Engine/Source/Editor/WorldPartitionEditor/Private/WorldPartition/Filter/WorldPartitionActorFilterTreeItems.cpp b/Engine/Source/Editor/WorldPartitionEditor/Private/WorldPartition/Filter/WorldPartitionActorFilterTreeItems.cpp
--- a/Engine/Source/Editor/WorldPartitionEditor/Private/WorldPartition/Filter/WorldPartitionActorFilterTreeItems.cpp
+++ b/Engine/Source/Editor/WorldPartitionEditor/Private/WorldPartition/Filter/WorldPartitionActorFilterTreeItems.cpp
@@ -60,6 +60,7 @@ struct SWorldPartitionActorFilterLabel : FSceneOutlinerCommonLabelData, public S
 				[
 					SNew(STextBlock)
 					.Text(this, &SWorldPartitionActorFilterLabel::GetDisplayText)
+					.ToolTipText(this, &SWorldPartitionActorFilterLabel::GetDisplayText)
 					.ColorAndOpacity(this, &SWorldPartitionActorFilterLabel::GetForegroundColor)
 				]
 			]
