Actor Filter: Do not focus search box when opened (steals focus from outliner)

#rnx

diff --git a/Engine/Source/Editor/WorldPartitionEditor/Private/WorldPartition/Filter/WorldPartitionActorFilterPropertyTypeCustomization.cpp b/Engine/Source/Editor/WorldPartitionEditor/Private/WorldPartition/Filter/WorldPartitionActorFilterPropertyTypeCustomization.cpp
--- a/Engine/Source/Editor/WorldPartitionEditor/Private/WorldPartition/Filter/WorldPartitionActorFilterPropertyTypeCustomization.cpp
+++ b/Engine/Source/Editor/WorldPartitionEditor/Private/WorldPartition/Filter/WorldPartitionActorFilterPropertyTypeCustomization.cpp
@@ -55,7 +55,7 @@ void FWorldPartitionActorFilterPropertyTypeCustomization::CustomizeHeader(TShare
 	FSceneOutlinerInitializationOptions InitOptions;
 	{
 		InitOptions.bShowHeaderRow = false;
-		InitOptions.bFocusSearchBoxWhenOpened = true;
+		InitOptions.bFocusSearchBoxWhenOpened = false;
 		InitOptions.bShowCreateNewFolder = false;
 		InitOptions.ModeFactory = ModeFactory;
 	}
