Fixed crash at exit of editor accessing ActorEditorContextSubsystem.

diff --git a/Engine/Source/Editor/UnrealEd/Private/SActorEditorContext.cpp b/Engine/Source/Editor/UnrealEd/Private/SActorEditorContext.cpp
--- a/Engine/Source/Editor/UnrealEd/Private/SActorEditorContext.cpp
+++ b/Engine/Source/Editor/UnrealEd/Private/SActorEditorContext.cpp
@@ -201,8 +201,11 @@ void SActorEditorContext::Rebuild()
 
 bool SActorEditorContext::IsVisible(UWorld* InWorld)
 {
-	TArray<IActorEditorContextClient*> Clients = UActorEditorContextSubsystem::Get()->GetDisplayableClients();
-	return InWorld && Clients.Num() > 0;
+	if (InWorld && GEditor)
+	{
+		return UActorEditorContextSubsystem::Get()->GetDisplayableClients().Num() > 0;
+	}
+	return false;
 }
 
 #undef LOCTEXT_NAMESPACE
\ No newline at end of file
diff --git a/Engine/Source/Editor/WorldPartitionEditor/Private/WorldPartition/ContentBundle/ContentBundleEditorSubsystem.cpp b/Engine/Source/Editor/WorldPartitionEditor/Private/WorldPartition/ContentBundle/ContentBundleEditorSubsystem.cpp
--- a/Engine/Source/Editor/WorldPartitionEditor/Private/WorldPartition/ContentBundle/ContentBundleEditorSubsystem.cpp
+++ b/Engine/Source/Editor/WorldPartitionEditor/Private/WorldPartition/ContentBundle/ContentBundleEditorSubsystem.cpp
@@ -18,12 +18,16 @@
 
 void UContentBundleEditingSubmodule::DoInitialize()
 {
+	check(GEditor);
 	UActorEditorContextSubsystem::Get()->RegisterClient(this);
 }
 
 void UContentBundleEditingSubmodule::DoDenitialize()
 {
-	UActorEditorContextSubsystem::Get()->UnregisterClient(this);
+	if (GEditor)
+	{
+		UActorEditorContextSubsystem::Get()->UnregisterClient(this);
+	}
 
 	EditingContentBundleGuid.Invalidate();
 	EditingContentBundlesStack.Empty();
