Fix log inside the commandlet output displayed in red after building HLODs - even if no error was reported
* Will now only pipe the output if the commandlet execution is considered a failure (if the actual commandlet code returned an error)
* Write the output line by line to avoid a red wall of text. This will allow lines with errors to be more easily spotted.
* Read from the commandlet log file to avoid the initial output being truncated (log file contains more info that what is written to StdOut)

diff --git a/Engine/Source/Editor/WorldPartitionEditor/Private/WorldPartitionEditorModule.cpp b/Engine/Source/Editor/WorldPartitionEditor/Private/WorldPartitionEditorModule.cpp
--- a/Engine/Source/Editor/WorldPartitionEditor/Private/WorldPartitionEditorModule.cpp
+++ b/Engine/Source/Editor/WorldPartitionEditor/Private/WorldPartitionEditorModule.cpp
@@ -24,6 +24,7 @@
 
 #include "Engine/Level.h"
 
+#include "Misc/FileHelper.h"
 #include "Misc/MessageDialog.h"
 #include "Misc/StringBuilder.h"
 #include "Misc/ScopedSlowTask.h"
@@ -374,11 +375,10 @@ static void LoadMap(const FString& MapToLoad)
 	}
 }
 
-void FWorldPartitionEditorModule::RunCommandletAsExternalProcess(const FString& InCommandletArgs, const FText& InOperationDescription, int32& OutResult, bool& bOutCancelled, FString& OutCommandletOutput)
+void FWorldPartitionEditorModule::RunCommandletAsExternalProcess(const FString& InCommandletArgs, const FText& InOperationDescription, int32& OutResult, bool& bOutCancelled)
 {
 	OutResult = 0;
 	bOutCancelled = false;
-	OutCommandletOutput.Empty();
 
 	FProcHandle ProcessHandle;
 
@@ -393,12 +393,23 @@ void FWorldPartitionEditorModule::RunCommandletAsExternalProcess(const FString&
 
 	// Try to provide complete Path, if we can't try with project name
 	FString ProjectPath = FPaths::IsProjectFilePathSet() ? FPaths::GetProjectFilePath() : FApp::GetProjectName();
+
+	// Obtain the log file path that will be used by the commandlet
+	FString LogFilePrefix = TEXT("Commandlet");
+	if (!FParse::Value(*InCommandletArgs, TEXT("Builder="), LogFilePrefix))
+	{
+		FParse::Value(*InCommandletArgs, TEXT("Run="), LogFilePrefix);
+	}
+	const FString TimeStamp = FString::Printf(TEXT("-%08x-%s"), FPlatformProcess::GetCurrentProcessId(), *FDateTime::Now().ToIso8601().Replace(TEXT(":"), TEXT(".")));
+	const FString RelLogFilePath = FPaths::ProjectLogDir() / TEXT("WorldPartition") / LogFilePrefix + TimeStamp + TEXT(".log");
+	const FString AbsLogFilePath = FPaths::ConvertRelativePathToFull(RelLogFilePath);
 		
 	TArray<FString> CommandletArgsArray;
 	CommandletArgsArray.Add(TEXT("\"") + ProjectPath + TEXT('"'));
 	CommandletArgsArray.Add(TEXT("-BaseDir=\"") + FString(FPlatformProcess::BaseDir()) + TEXT('"'));
 	CommandletArgsArray.Add(TEXT("-Unattended"));
 	CommandletArgsArray.Add(TEXT("-RunningFromUnrealEd"));
+	CommandletArgsArray.Add(TEXT("-AbsLog=\"") + AbsLogFilePath + TEXT('"'));
 	CommandletArgsArray.Add(InCommandletArgs);
 
 	OnExecuteCommandletEvent.Broadcast(CommandletArgsArray);
@@ -413,10 +424,10 @@ void FWorldPartitionEditorModule::RunCommandletAsExternalProcess(const FString&
 	UE_LOG(LogWorldPartitionEditor, Display, TEXT("Running commandlet: %s %s"), *CurrentExecutableName, *Arguments);
 
 	uint32 ProcessID;
-	const bool bLaunchDetached = true;
+	const bool bLaunchDetached = false;
 	const bool bLaunchHidden = true;
 	const bool bLaunchReallyHidden = true;
-	ProcessHandle = FPlatformProcess::CreateProc(*CurrentExecutableName, *Arguments, bLaunchDetached, bLaunchHidden, bLaunchReallyHidden, &ProcessID, 0, nullptr, WritePipe, ReadPipe);
+	ProcessHandle = FPlatformProcess::CreateProc(*CurrentExecutableName, *Arguments, bLaunchDetached, bLaunchHidden, bLaunchReallyHidden, &ProcessID, 0, nullptr, WritePipe);
 
 	while (FPlatformProcess::IsProcRunning(ProcessHandle))
 	{
@@ -427,11 +438,7 @@ void FWorldPartitionEditorModule::RunCommandletAsExternalProcess(const FString&
 			break;
 		}
 
-		const FString LogString = FPlatformProcess::ReadPipe(ReadPipe);
-		if (!LogString.IsEmpty())
-		{
-			OutCommandletOutput.Append(LogString);
-		}
+		FString LogString = FPlatformProcess::ReadPipe(ReadPipe);
 
 		// Parse output, look for progress indicator in the log (in the form "Display: [i / N] Msg...\n")
 		const FRegexPattern LogProgressPattern(TEXT("Display:\\s\\[([0-9]+)\\s\\/\\s([0-9]+)\\]\\s(.+)?(?=\\.{3}$)"));
@@ -448,22 +455,30 @@ void FWorldPartitionEditorModule::RunCommandletAsExternalProcess(const FString&
 		FPlatformProcess::Sleep(0.1);
 	}
 
-
 	FPlatformProcess::GetProcReturnCode(ProcessHandle, &OutResult);
-
-	UE_LOG(LogWorldPartitionEditor, Display, TEXT("#### Begin commandlet output ####\n%s"), *OutCommandletOutput);
-
 	FPlatformProcess::ClosePipe(ReadPipe, WritePipe);
 
 	if (OutResult == 0)
 	{
-		UE_LOG(LogWorldPartitionEditor, Display, TEXT("#### Commandlet Succeeded ####"));
+		UE_LOG(LogWorldPartitionEditor, Display, TEXT("Commandlet executed successfully."));
+		UE_LOG(LogWorldPartitionEditor, Display, TEXT("Detailed output can be found in %s"), *AbsLogFilePath);
 	}
 	else
 	{
 		UE_LOG(LogWorldPartitionEditor, Error, TEXT("#### Commandlet Failed ####"));
 		UE_LOG(LogWorldPartitionEditor, Error, TEXT("%s %s"), *CurrentExecutableName, *Arguments);
 		UE_LOG(LogWorldPartitionEditor, Error, TEXT("Return Code: %i"), OutResult);
+
+		UE_LOG(LogWorldPartitionEditor, Display, TEXT("#### BEGIN COMMANDLET OUTPUT ####"));
+
+		TArray<FString> OutputLines;
+		FFileHelper::LoadFileToStringArray(OutputLines, *AbsLogFilePath);
+		for (const FString& OutputLine : OutputLines)
+		{
+			UE_LOG(LogWorldPartitionEditor, Display, TEXT("#### COMMANDLET OUTPUT >> %s"), *OutputLine);
+		}		
+
+		UE_LOG(LogWorldPartitionEditor, Display, TEXT("#### END COMMANDLET OUTPUT ####"));
 	}
 }
 
@@ -518,8 +533,7 @@ bool FWorldPartitionEditorModule::ConvertMap(const FString& InLongPackageName)
 		
 		int32 Result;
 		bool bCancelled;
-		FString CommandletOutput;
-		RunCommandletAsExternalProcess(CommandletArgs, OperationDescription, Result, bCancelled, CommandletOutput);
+		RunCommandletAsExternalProcess(CommandletArgs, OperationDescription, Result, bCancelled);
 		if (!bCancelled && Result == 0)
 		{	
 #if	PLATFORM_DESKTOP
@@ -660,10 +674,9 @@ bool FWorldPartitionEditorModule::Build(const FRunBuilderParams& InParams)
 
 	int32 Result;
 	bool bCancelled;
-	FString CommandletOutput;
 
 	FString CommandletArgs(CommandletArgsBuilder.ToString());
-	RunCommandletAsExternalProcess(CommandletArgs, OperationDescription, Result, bCancelled, CommandletOutput);
+	RunCommandletAsExternalProcess(CommandletArgs, OperationDescription, Result, bCancelled);
 
 	RescanAssets(MapPackage);
 
diff --git a/Engine/Source/Editor/WorldPartitionEditor/Public/WorldPartitionEditorModule.h b/Engine/Source/Editor/WorldPartitionEditor/Public/WorldPartitionEditorModule.h
--- a/Engine/Source/Editor/WorldPartitionEditor/Public/WorldPartitionEditorModule.h
+++ b/Engine/Source/Editor/WorldPartitionEditor/Public/WorldPartitionEditorModule.h
@@ -131,7 +131,7 @@ private:
 	bool BuildLandscapeSplineMeshes(UWorld* InWorld);
 
 private:
-	void RunCommandletAsExternalProcess(const FString& InCommandletArgs, const FText& InOperationDescription, int32& OutResult, bool& bOutCancelled, FString& OutCommandletOutput);
+	void RunCommandletAsExternalProcess(const FString& InCommandletArgs, const FText& InOperationDescription, int32& OutResult, bool& bOutCancelled);
 	void OnConvertMap();
 
 	FDelegateHandle LevelEditorExtenderDelegateHandle;
