World Partition: Fix crash when using BugItGo command when the world partition editor is opened but never had focus once.

diff --git a/Engine/Source/Editor/WorldPartitionEditor/Private/WorldPartition/SWorldPartitionEditorGrid2D.cpp b/Engine/Source/Editor/WorldPartitionEditor/Private/WorldPartition/SWorldPartitionEditorGrid2D.cpp
--- a/Engine/Source/Editor/WorldPartitionEditor/Private/WorldPartition/SWorldPartitionEditorGrid2D.cpp
+++ b/Engine/Source/Editor/WorldPartitionEditor/Private/WorldPartition/SWorldPartitionEditorGrid2D.cpp
@@ -2029,19 +2029,20 @@ FReply SWorldPartitionEditorGrid2D::FocusLoadedRegions()
 
 void SWorldPartitionEditorGrid2D::FocusBox(const FBox& Box) const
 {
-	check(ScreenRect.bIsValid);
+	if (ScreenRect.bIsValid)
+	{
+		const FBox2D Box2D(FVector2D(Box.Min), FVector2D(Box.Max));
+		Trans = -FVector2D(Box2D.GetCenter());
 
-	const FBox2D Box2D(FVector2D(Box.Min), FVector2D(Box.Max));
-	Trans = -FVector2D(Box2D.GetCenter());
+		if (Box2D.GetArea() > 0.0f)
+		{
+			const FVector2D ScreenExtent = ScreenRect.GetExtent();
+			const FVector2D SelectExtent = FVector2D(Box2D.GetExtent());
+			Scale = (ScreenExtent / SelectExtent).GetMin() * 0.85f;
+		}
 
-	if (Box2D.GetArea() > 0.0f)
-	{
-		const FVector2D ScreenExtent = ScreenRect.GetExtent();
-		const FVector2D SelectExtent = FVector2D(Box2D.GetExtent());
-		Scale = (ScreenExtent / SelectExtent).GetMin() * 0.85f;
+		UpdateTransform();
 	}
-
-	UpdateTransform();
 }
 
 void SWorldPartitionEditorGrid2D::UpdateTransform() const
