When converting a world to world partition, reconstruct the world partition editor to reflect the change.

diff --git a/Engine/Source/Editor/WorldPartitionEditor/Private/WorldPartition/SWorldPartitionEditor.cpp b/Engine/Source/Editor/WorldPartitionEditor/Private/WorldPartition/SWorldPartitionEditor.cpp
--- a/Engine/Source/Editor/WorldPartitionEditor/Private/WorldPartition/SWorldPartitionEditor.cpp
+++ b/Engine/Source/Editor/WorldPartitionEditor/Private/WorldPartition/SWorldPartitionEditor.cpp
@@ -23,6 +23,7 @@ void SWorldPartitionEditor::Construct(const FArguments& InArgs)
 
 	FWorldBrowserModule& WorldBrowserModule = FModuleManager::LoadModuleChecked<FWorldBrowserModule>("WorldBrowser");
 	WorldBrowserModule.OnBrowseWorld.AddSP(this, &SWorldPartitionEditor::OnBrowseWorld);
+	UWorldPartition::WorldPartitionChangedEvent.AddSP(this, &SWorldPartitionEditor::OnBrowseWorld);
 
 	IWorldPartitionEditorModule& WorldPartitionEditorModule = FModuleManager::LoadModuleChecked<IWorldPartitionEditorModule>("WorldPartitionEditor");
 	WorldPartitionEditorModule.OnWorldPartitionCreated().AddSP(this, &SWorldPartitionEditor::OnBrowseWorld);
@@ -32,6 +33,7 @@ SWorldPartitionEditor::~SWorldPartitionEditor()
 {
 	FWorldBrowserModule& WorldBrowserModule = FModuleManager::GetModuleChecked<FWorldBrowserModule>("WorldBrowser");
 	WorldBrowserModule.OnBrowseWorld.RemoveAll(this);
+	UWorldPartition::WorldPartitionChangedEvent.RemoveAll(this);
 	
 	IWorldPartitionEditorModule& WorldPartitionEditorModule = FModuleManager::LoadModuleChecked<IWorldPartitionEditorModule>("WorldPartitionEditor");
 	WorldPartitionEditorModule.OnWorldPartitionCreated().RemoveAll(this);
diff --git a/Engine/Source/Runtime/Engine/Private/WorldPartition/WorldPartition.cpp b/Engine/Source/Runtime/Engine/Private/WorldPartition/WorldPartition.cpp
--- a/Engine/Source/Runtime/Engine/Private/WorldPartition/WorldPartition.cpp
+++ b/Engine/Source/Runtime/Engine/Private/WorldPartition/WorldPartition.cpp
@@ -257,6 +257,7 @@ public:
 		bIncludeNonSpatiallyLoadedActors = true;
 	}
 };
+UWorldPartition::FWorldPartitionChangedEvent UWorldPartition::WorldPartitionChangedEvent;
 #endif
 
 UWorldPartition::UWorldPartition(const FObjectInitializer& ObjectInitializer)
diff --git a/Engine/Source/Runtime/Engine/Private/WorldPartition/WorldPartitionConverter.cpp b/Engine/Source/Runtime/Engine/Private/WorldPartition/WorldPartitionConverter.cpp
--- a/Engine/Source/Runtime/Engine/Private/WorldPartition/WorldPartitionConverter.cpp
+++ b/Engine/Source/Runtime/Engine/Private/WorldPartition/WorldPartitionConverter.cpp
@@ -176,6 +176,7 @@ bool FWorldPartitionConverter::Convert()
 	if (bCreatedWorldPartition)
 	{
 		WorldPartition->Initialize(World, FTransform::Identity);
+		UWorldPartition::WorldPartitionChangedEvent.Broadcast(World);
 	}
 
 	return true;
diff --git a/Engine/Source/Runtime/Engine/Public/WorldPartition/WorldPartition.h b/Engine/Source/Runtime/Engine/Public/WorldPartition/WorldPartition.h
--- a/Engine/Source/Runtime/Engine/Public/WorldPartition/WorldPartition.h
+++ b/Engine/Source/Runtime/Engine/Public/WorldPartition/WorldPartition.h
@@ -216,6 +216,10 @@ public:
 	void FlushStreaming();
 	URuntimeHashExternalStreamingObjectBase* FlushStreamingToExternalStreamingObject(const FString& ExternalStreamingObjectName);
 
+	// Event when world partition was enabled/disabled in the world
+	DECLARE_MULTICAST_DELEGATE_OneParam(FWorldPartitionChangedEvent, UWorld*);
+	static FWorldPartitionChangedEvent WorldPartitionChangedEvent;
+
 	DECLARE_MULTICAST_DELEGATE_OneParam(FWorldPartitionGenerateStreamingDelegate, TArray<FString>*);
 	FWorldPartitionGenerateStreamingDelegate OnPreGenerateStreaming;
 
