* [SceneOutliner] Using Load and Unload selected volumes context menu options for LocationVolume in the Outliner does not grey out these options.

diff --git a/Engine/Source/Editor/WorldPartitionEditor/Private/WorldPartitionEditorModule.cpp b/Engine/Source/Editor/WorldPartitionEditor/Private/WorldPartitionEditorModule.cpp
--- a/Engine/Source/Editor/WorldPartitionEditor/Private/WorldPartitionEditorModule.cpp
+++ b/Engine/Source/Editor/WorldPartitionEditor/Private/WorldPartitionEditorModule.cpp
@@ -83,15 +83,38 @@ static void OnSelectedWorldPartitionVolumesToggleLoading(TArray<TWeakObjectPtr<A
 	}
 }
 
+static bool CanLoadUnloadSelectedVolumes(TArray<TWeakObjectPtr<AActor>> Volumes, bool bLoad)
+{
+	for (TWeakObjectPtr<AActor> Actor : Volumes)
+	{
+		if (Actor->Implements<UWorldPartitionActorLoaderInterface>())
+		{
+			IWorldPartitionActorLoaderInterface::ILoaderAdapter* LoaderAdapter = Cast<IWorldPartitionActorLoaderInterface>(Actor)->GetLoaderAdapter();
+			if (bLoad != LoaderAdapter->IsLoaded())
+			{
+				return true;
+			}
+		}
+	}
+
+	return false;
+}
+
 static void CreateLevelViewportContextMenuEntries(FMenuBuilder& MenuBuilder, TArray<TWeakObjectPtr<AActor>> Volumes)
 {
 	MenuBuilder.BeginSection("WorldPartition", LOCTEXT("WorldPartition", "World Partition"));
-
+	
 	MenuBuilder.AddMenuEntry(
 		LOCTEXT("WorldPartitionLoad", "Load selected volumes"),
 		LOCTEXT("WorldPartitionLoad_Tooltip", "Load selected volumes"),
 		FSlateIcon(),
-		FExecuteAction::CreateStatic(OnSelectedWorldPartitionVolumesToggleLoading, Volumes, true),
+		FUIAction(
+			FExecuteAction::CreateStatic(OnSelectedWorldPartitionVolumesToggleLoading, Volumes, true),
+			FCanExecuteAction::CreateLambda([Volumes]
+			{
+				return CanLoadUnloadSelectedVolumes(Volumes, true);
+			})
+		),
 		NAME_None,
 		EUserInterfaceActionType::Button);
 
@@ -99,10 +122,16 @@ static void CreateLevelViewportContextMenuEntries(FMenuBuilder& MenuBuilder, TAr
 		LOCTEXT("WorldPartitionUnload", "Unload selected volumes"),
 		LOCTEXT("WorldPartitionUnload_Tooltip", "Load selected volumes"),
 		FSlateIcon(),
-		FExecuteAction::CreateStatic(OnSelectedWorldPartitionVolumesToggleLoading, Volumes, false),
+		FUIAction(
+			FExecuteAction::CreateStatic(OnSelectedWorldPartitionVolumesToggleLoading, Volumes, false),
+			FCanExecuteAction::CreateLambda([Volumes]
+				{
+					return CanLoadUnloadSelectedVolumes(Volumes, false);
+				})
+		),
 		NAME_None,
 		EUserInterfaceActionType::Button);
-
+	
 	MenuBuilder.EndSection();
 }
 
