World Partition Editor QoL
* Make sure we can only play with the Loaded Regions & related streaming features when the streaming is enabled. Hide the loaded regions when we toggle the Streaming enable so we keep them.
* When converting a Loaded Region to LocationVolume, make sure we have an incremental label for the user.
* Make sure the scaling gizmo would be intuitive on the Location Volume (after converting a loaded region) by scaling the brush instead of the actor.
* Fixed : Converting a loaded region into an asset does not count as a transaction on the undo history, but deleting it from the viewport or the outliner does.

diff --git a/Engine/Source/Editor/WorldPartitionEditor/Private/WorldPartition/SWorldPartitionEditorGrid2D.cpp b/Engine/Source/Editor/WorldPartitionEditor/Private/WorldPartition/SWorldPartitionEditorGrid2D.cpp
--- a/Engine/Source/Editor/WorldPartitionEditor/Private/WorldPartition/SWorldPartitionEditorGrid2D.cpp
+++ b/Engine/Source/Editor/WorldPartitionEditor/Private/WorldPartition/SWorldPartitionEditorGrid2D.cpp
@@ -165,10 +165,16 @@ static bool HitTestZFromLocation(UWorld* World, const FVector2D& WorldLocation,
 	return World->LineTraceSingleByChannel(OutResult, TraceStart, TraceEnd, ECC_Camera, TraceParams);
 }
 
+const TSet<TObjectPtr<UWorldPartitionEditorLoaderAdapter>>& GetRegisteredEditorLoaderAdapters(UWorldPartition* WorldPartition)
+{
+	static const TSet<TObjectPtr<UWorldPartitionEditorLoaderAdapter>> EmptyRegisteredEditorLoaderAdapters;
+	return WorldPartition->IsStreamingEnabled() ? WorldPartition->GetRegisteredEditorLoaderAdapters() : EmptyRegisteredEditorLoaderAdapters;
+}
+
 template <class T>
 void ForEachIntersectingLoaderAdapters(UWorldPartition* WorldPartition, const FBox& SelectBox, T Func)
 {
-	for (UWorldPartitionEditorLoaderAdapter* EditorLoaderAdapter : WorldPartition->GetRegisteredEditorLoaderAdapters())
+	for (UWorldPartitionEditorLoaderAdapter* EditorLoaderAdapter : GetRegisteredEditorLoaderAdapters(WorldPartition))
 	{
 		if (IWorldPartitionActorLoaderInterface::ILoaderAdapter* LoaderAdapter = EditorLoaderAdapter->GetLoaderAdapter())
 		{
@@ -477,14 +483,19 @@ void SWorldPartitionEditorGrid2D::Construct(const FArguments& InArgs)
 
 	auto CanLoadSelectedRegions = [this, CanLoadUnloadSelectedRegions]() { return CanLoadUnloadSelectedRegions(true); };
 	auto CanUnloadSelectedRegions = [this, CanLoadUnloadSelectedRegions]() { return CanLoadUnloadSelectedRegions(false); };
-		
+
+	auto CanLoadFromHere = [this]()
+	{
+		return GetWorldPartition()->IsStreamingEnabled();
+	};
+
 	ActionList.MapAction(Commands.CreateRegionFromSelection, FExecuteAction::CreateSP(this, &SWorldPartitionEditorGrid2D::CreateRegionFromSelection), FCanExecuteAction::CreateLambda(CanCreateRegionFromSelection));
 	ActionList.MapAction(Commands.LoadSelectedRegions, FExecuteAction::CreateSP(this, &SWorldPartitionEditorGrid2D::LoadSelectedRegions), FCanExecuteAction::CreateLambda(CanLoadSelectedRegions));
 	ActionList.MapAction(Commands.UnloadSelectedRegions, FExecuteAction::CreateSP(this, &SWorldPartitionEditorGrid2D::UnloadSelectedRegions), FCanExecuteAction::CreateLambda(CanUnloadSelectedRegions));
 	ActionList.MapAction(Commands.ConvertSelectedRegionsToActors, FExecuteAction::CreateSP(this, &SWorldPartitionEditorGrid2D::ConvertSelectedRegionsToActors), FCanExecuteAction::CreateLambda(CanConvertSelectedRegionsToActors));
 	ActionList.MapAction(Commands.MoveCameraHere, FExecuteAction::CreateSP(this, &SWorldPartitionEditorGrid2D::MoveCameraHere));
 	ActionList.MapAction(Commands.PlayFromHere, FExecuteAction::CreateSP(this, &SWorldPartitionEditorGrid2D::PlayFromHere));
-	ActionList.MapAction(Commands.LoadFromHere, FExecuteAction::CreateSP(this, &SWorldPartitionEditorGrid2D::LoadFromHere));
+	ActionList.MapAction(Commands.LoadFromHere, FExecuteAction::CreateSP(this, &SWorldPartitionEditorGrid2D::LoadFromHere), FCanExecuteAction::CreateLambda(CanLoadFromHere));
 	ActionList.MapAction(Commands.BugItHere, FExecuteAction::CreateSP(this, &SWorldPartitionEditorGrid2D::BugItHere));
 }
 
@@ -562,6 +573,7 @@ void SWorldPartitionEditorGrid2D::ConvertSelectedRegionsToActors()
 	FLoaderInterfaceSet TmpSelectedLoaderInterfaces = MoveTemp(SelectedLoaderInterfaces);
 	ClearSelection();
 
+	const FScopedTransaction Transaction(LOCTEXT("ConvertSelectedRegionsToActors", "Convert Selected Region(s) to Actor(s)"));
 	const FBox WorldBounds = GetWorldPartition()->GetRuntimeWorldBounds();
 	for (const TWeakInterfacePtr<IWorldPartitionActorLoaderInterface>& SelectedLoaderAdapter : TmpSelectedLoaderInterfaces)
 	{
@@ -573,17 +585,19 @@ void SWorldPartitionEditorGrid2D::ConvertSelectedRegionsToActors()
 			const FBox ActorVolumeBox(FVector(LoaderVolumeBox.Min.X, LoaderVolumeBox.Min.Y, WorldBounds.Min.Z), FVector(LoaderVolumeBox.Max.X, LoaderVolumeBox.Max.Y, WorldBounds.Max.Z));
 
 			ALocationVolume* LocationVolume = GetWorld()->SpawnActor<ALocationVolume>(ActorVolumeBox.GetCenter(), FRotator::ZeroRotator);
-
+			LocationVolume->Modify();
+			FActorLabelUtilities::SetActorLabelUnique(LocationVolume, LocationVolume->GetActorLabel());
+			
 			UCubeBuilder* Builder = NewObject<UCubeBuilder>();
-			Builder->X = 1.0f;
-			Builder->Y = 1.0f;
-			Builder->Z = 1.0f;
-			UActorFactory::CreateBrushForVolumeActor(LocationVolume, Builder);
-
-			LocationVolume->GetRootComponent()->SetWorldScale3D(ActorVolumeBox.GetSize());
+			Builder->Modify();
+			FVector Extent = ActorVolumeBox.GetExtent();
+			Builder->X = Extent.X * 2;
+			Builder->Y = Extent.Y * 2;
+			Builder->Z = FMath::Max(Extent.Z * 2, 200.0f);
 
+			UActorFactory::CreateBrushForVolumeActor(LocationVolume, Builder);
+			
 			LocationVolume->GetLoaderAdapter()->Load();
-
 			LoaderAdapter->Unload();
 			
 			GetWorldPartition()->ReleaseEditorLoaderAdapter(EditorLoaderAdapter);
@@ -1033,7 +1047,7 @@ void SWorldPartitionEditorGrid2D::Tick(const FGeometry& AllottedGeometry, const
 	DirtyActorGuids.Reset();
 	ShownLoaderInterfaces.Reset(); 
 	
-	for (UWorldPartitionEditorLoaderAdapter* EditorLoaderAdapter : GetWorldPartition()->GetRegisteredEditorLoaderAdapters())
+	for (UWorldPartitionEditorLoaderAdapter* EditorLoaderAdapter : GetRegisteredEditorLoaderAdapters(GetWorldPartition()))
 	{
 		check(EditorLoaderAdapter);
 		IWorldPartitionActorLoaderInterface::ILoaderAdapter* LoaderAdapter = EditorLoaderAdapter->GetLoaderAdapter();
@@ -2001,7 +2015,7 @@ FReply SWorldPartitionEditorGrid2D::FocusLoadedRegions()
 {
 	FBox SelectionBox(ForceInit);
 
-	for (UWorldPartitionEditorLoaderAdapter* EditorLoaderAdapter : GetWorldPartition()->GetRegisteredEditorLoaderAdapters())
+	for (UWorldPartitionEditorLoaderAdapter* EditorLoaderAdapter : GetRegisteredEditorLoaderAdapters(GetWorldPartition()))
 	{
 		check(EditorLoaderAdapter);
 		IWorldPartitionActorLoaderInterface::ILoaderAdapter* LoaderAdapter = EditorLoaderAdapter->GetLoaderAdapter();
