Add missing include (ScopedTransaction)

diff --git a/Engine/Source/Editor/WorldPartitionEditor/Private/WorldPartition/SWorldPartitionEditorGrid2D.cpp b/Engine/Source/Editor/WorldPartitionEditor/Private/WorldPartition/SWorldPartitionEditorGrid2D.cpp
--- a/Engine/Source/Editor/WorldPartitionEditor/Private/WorldPartition/SWorldPartitionEditorGrid2D.cpp
+++ b/Engine/Source/Editor/WorldPartitionEditor/Private/WorldPartition/SWorldPartitionEditorGrid2D.cpp
@@ -17,6 +17,7 @@
 #include "LocationVolume.h"
 #include "Modules/ModuleManager.h"
 #include "Rendering/SlateRenderer.h"
+#include "ScopedTransaction.h"
 #include "SWorldPartitionViewportWidget.h"
 #include "TextureResource.h"
 #include "Widgets/Input/SButton.h"
