World Partition Editor: don't highlight loader adapters that aren't loaded.

diff --git a/Engine/Source/Editor/WorldPartitionEditor/Private/WorldPartition/SWorldPartitionEditorGrid2D.cpp b/Engine/Source/Editor/WorldPartitionEditor/Private/WorldPartition/SWorldPartitionEditorGrid2D.cpp
--- a/Engine/Source/Editor/WorldPartitionEditor/Private/WorldPartition/SWorldPartitionEditorGrid2D.cpp
+++ b/Engine/Source/Editor/WorldPartitionEditor/Private/WorldPartition/SWorldPartitionEditorGrid2D.cpp
@@ -1804,7 +1804,7 @@ int32 SWorldPartitionEditorGrid2D::PaintMinimap(const FGeometry& AllottedGeometr
 				{
 					const IWorldPartitionActorLoaderInterface::ILoaderAdapter* LoaderAdapter = LoaderInterface->GetLoaderAdapter();
 
-					if (LoaderAdapter->GetBoundingBox().IsSet())
+					if (LoaderAdapter->IsLoaded() && LoaderAdapter->GetBoundingBox().IsSet())
 					{
 						const FBox AdapterBounds = *LoaderAdapter->GetBoundingBox();
 
