World Partition Editor Loading:
- Changed Editor Settings 'DisableLoadingInEditor' to 'EnableLoadingInEditor' to be easier to understand for entry-level users.
- Actors loading state will now react to changing Editor Settings 'EnableLoadingInEditor' property.

diff --git a/Engine/Source/Editor/UnrealEd/Private/SWorldPartitionViewportWidget.cpp b/Engine/Source/Editor/UnrealEd/Private/SWorldPartitionViewportWidget.cpp
--- a/Engine/Source/Editor/UnrealEd/Private/SWorldPartitionViewportWidget.cpp
+++ b/Engine/Source/Editor/UnrealEd/Private/SWorldPartitionViewportWidget.cpp
@@ -78,7 +78,7 @@ EVisibility SWorldPartitionViewportWidget::GetVisibility(UWorld* InWorld)
 		UWorldPartition* WorldPartition = InWorld->GetWorldPartition();
 		if (WorldPartition->IsStreamingEnabled())
 		{
-			if (!WorldPartition->HasLoadedUserCreatedRegions() && !WorldPartitionEditorModule->GetDisableLoadingInEditor())
+			if (!WorldPartition->HasLoadedUserCreatedRegions() && WorldPartitionEditorModule->GetEnableLoadingInEditor())
 			{
 				Message = LOCTEXT("NoRegionsLoadedText","No Regions Loaded");
 				Tooltip = LOCTEXT("NoRegionsLoadedToolTip", "To load a region, drag select an area in the World Partition map and choose 'Load Region From Selection' from the context menu.");
diff --git a/Engine/Source/Editor/WorldPartitionEditor/Private/WorldPartition/SWorldPartitionEditorGrid2D.cpp b/Engine/Source/Editor/WorldPartitionEditor/Private/WorldPartition/SWorldPartitionEditorGrid2D.cpp
--- a/Engine/Source/Editor/WorldPartitionEditor/Private/WorldPartition/SWorldPartitionEditorGrid2D.cpp
+++ b/Engine/Source/Editor/WorldPartitionEditor/Private/WorldPartition/SWorldPartitionEditorGrid2D.cpp
@@ -368,7 +368,7 @@ void SWorldPartitionEditorGrid2D::Construct(const FArguments& InArgs)
 			ToolbarBuilder.AddWidget(SNew(SCheckBox)
 				.ForegroundColor(FSlateColor::UseForeground())
 				.IsChecked(GetMutableDefault<UWorldPartitionEditorPerProjectUserSettings>()->GetBugItGoLoadRegion() ? ECheckBoxState::Checked : ECheckBoxState::Unchecked)
-				.Visibility_Lambda([this]() { return GetDefault<UWorldPartitionEditorSettings>()->bDisableLoadingInEditor ? EVisibility::Hidden : EVisibility::Visible; })
+				.Visibility_Lambda([this]() { return GetDefault<UWorldPartitionEditorSettings>()->bEnableLoadingInEditor ? EVisibility::Visible : EVisibility::Hidden; })
 				.OnCheckStateChanged(FOnCheckStateChanged::CreateLambda([=](ECheckBoxState State) { GetMutableDefault<UWorldPartitionEditorPerProjectUserSettings>()->SetBugItGoLoadRegion(State == ECheckBoxState::Checked); }))
 				[
 					SNew(SBox)
@@ -415,7 +415,7 @@ void SWorldPartitionEditorGrid2D::Construct(const FArguments& InArgs)
 				.Text(LOCTEXT("FocusLoadedRegions", "Focus Loaded Regions"))
 				.OnClicked(this, &SWorldPartitionEditorGrid2D::FocusLoadedRegions)
 				.IsEnabled_Lambda([this]() { return IsInteractive() && GetWorldPartition() && GetWorldPartition()->HasLoadedUserCreatedRegions(); })
-				.Visibility_Lambda([this]() { return GetDefault<UWorldPartitionEditorSettings>()->bDisableLoadingInEditor ? EVisibility::Hidden : EVisibility::Visible; }));
+				.Visibility_Lambda([this]() { return GetDefault<UWorldPartitionEditorSettings>()->bEnableLoadingInEditor ? EVisibility::Visible : EVisibility::Hidden; }));
 		}
 
 		return ToolbarBuilder.MakeWidget();
@@ -782,7 +782,7 @@ FReply SWorldPartitionEditorGrid2D::OnMouseButtonUp(const FGeometry& MyGeometry,
 
 			const FEditorCommands& Commands = FEditorCommands::Get();
 
-			if (!GetDefault<UWorldPartitionEditorSettings>()->bDisableLoadingInEditor)
+			if (GetDefault<UWorldPartitionEditorSettings>()->bEnableLoadingInEditor)
 			{
 				MenuBuilder.BeginSection(NAME_None, LOCTEXT("WorldPartitionSelection", "Selection"));
 					MenuBuilder.AddMenuEntry(Commands.CreateRegionFromSelection);
@@ -807,7 +807,7 @@ FReply SWorldPartitionEditorGrid2D::OnMouseButtonUp(const FGeometry& MyGeometry,
 					MenuBuilder.AddMenuEntry(Commands.PlayFromHere);
 				}
 				
-				if (!GetDefault<UWorldPartitionEditorSettings>()->bDisableLoadingInEditor)
+				if (GetDefault<UWorldPartitionEditorSettings>()->bEnableLoadingInEditor)
 				{
 					MenuBuilder.AddMenuEntry(Commands.LoadFromHere);
 				}
@@ -890,7 +890,7 @@ FReply SWorldPartitionEditorGrid2D::OnMouseButtonDoubleClick(const FGeometry& In
 		{
 			MoveCameraHere();
 
-			if (InMouseEvent.IsControlDown() && !GetDefault<UWorldPartitionEditorSettings>()->bDisableLoadingInEditor)
+			if (InMouseEvent.IsControlDown() && GetDefault<UWorldPartitionEditorSettings>()->bEnableLoadingInEditor)
 			{
 				LoadFromHere();
 			}
diff --git a/Engine/Source/Editor/WorldPartitionEditor/Private/WorldPartition/WorldPartitionEditorSettings.cpp b/Engine/Source/Editor/WorldPartitionEditor/Private/WorldPartition/WorldPartitionEditorSettings.cpp
new file mode 100644
--- /dev/null
+++ b/Engine/Source/Editor/WorldPartitionEditor/Private/WorldPartition/WorldPartitionEditorSettings.cpp
@@ -0,0 +1,17 @@
+// Copyright Epic Games, Inc. All Rights Reserved.
+
+#include "WorldPartition/WorldPartitionEditorSettings.h"
+#include "WorldPartition/WorldPartition.h"
+
+void UWorldPartitionEditorSettings::PostEditChangeProperty(struct FPropertyChangedEvent& PropertyChangedEvent)
+{
+	Super::PostEditChangeProperty(PropertyChangedEvent);
+
+	if (PropertyChangedEvent.Property && PropertyChangedEvent.Property->GetFName() == GET_MEMBER_NAME_CHECKED(UWorldPartitionEditorSettings, bEnableLoadingInEditor))
+	{	
+		if (UWorldPartition* WorldPartition = GWorld ? GWorld->GetWorldPartition() : nullptr)
+		{
+			WorldPartition->OnEnableLoadingInEditorChanged();
+		}
+	}
+}
\ No newline at end of file
diff --git a/Engine/Source/Editor/WorldPartitionEditor/Private/WorldPartitionEditorModule.cpp b/Engine/Source/Editor/WorldPartitionEditor/Private/WorldPartitionEditorModule.cpp
--- a/Engine/Source/Editor/WorldPartitionEditor/Private/WorldPartitionEditorModule.cpp
+++ b/Engine/Source/Editor/WorldPartitionEditor/Private/WorldPartitionEditorModule.cpp
@@ -287,14 +287,14 @@ int32 FWorldPartitionEditorModule::GetMinimapLowQualityWorldUnitsPerPixelThresho
 	return GetDefault<UWorldPartitionEditorSettings>()->MinimapLowQualityWorldUnitsPerPixelThreshold;
 }
 
-bool FWorldPartitionEditorModule::GetDisableLoadingInEditor() const
+bool FWorldPartitionEditorModule::GetEnableLoadingInEditor() const
 {
-	return GetDefault<UWorldPartitionEditorSettings>()->bDisableLoadingInEditor;
+	return GetDefault<UWorldPartitionEditorSettings>()->bEnableLoadingInEditor;
 }
 
-void FWorldPartitionEditorModule::SetDisableLoadingInEditor(bool bInDisableLoadingInEditor)
+void FWorldPartitionEditorModule::SetEnableLoadingInEditor(bool bInEnableLoadingInEditor)
 {
-	GetMutableDefault<UWorldPartitionEditorSettings>()->bDisableLoadingInEditor = bInDisableLoadingInEditor;
+	GetMutableDefault<UWorldPartitionEditorSettings>()->bEnableLoadingInEditor = bInEnableLoadingInEditor;
 }
 
 bool FWorldPartitionEditorModule::GetDisablePIE() const
@@ -828,7 +828,7 @@ UWorldPartitionEditorSettings::UWorldPartitionEditorSettings()
 	CommandletClass = UWorldPartitionConvertCommandlet::StaticClass();
 	InstancedFoliageGridSize = 25600;
 	MinimapLowQualityWorldUnitsPerPixelThreshold = 12800;
-	bDisableLoadingInEditor = false;
+	bEnableLoadingInEditor = true;
 	bDisablePIE = false;
 	bDisableBugIt = false;
 	bAdvancedMode = true;
diff --git a/Engine/Source/Editor/WorldPartitionEditor/Public/WorldPartition/WorldPartitionEditorSettings.h b/Engine/Source/Editor/WorldPartitionEditor/Public/WorldPartition/WorldPartitionEditorSettings.h
--- a/Engine/Source/Editor/WorldPartitionEditor/Public/WorldPartition/WorldPartitionEditorSettings.h
+++ b/Engine/Source/Editor/WorldPartitionEditor/Public/WorldPartition/WorldPartitionEditorSettings.h
@@ -15,6 +15,8 @@ class WORLDPARTITIONEDITOR_API UWorldPartitionEditorSettings : public UDeveloper
 public:
 	UWorldPartitionEditorSettings();
 
+	virtual void PostEditChangeProperty(FPropertyChangedEvent& PropertyChangedEvent) override;
+
 	UPROPERTY(Config, EditAnywhere, Category = MapConversion, Meta = (ToolTip = "Commandlet class to use for World Parition conversion"))
 	TSubclassOf<UWorldPartitionConvertCommandlet> CommandletClass;
 
@@ -25,7 +27,7 @@ public:
 	int32 MinimapLowQualityWorldUnitsPerPixelThreshold;
 
 	UPROPERTY(Config, EditAnywhere, Category = WorldPartition, Meta = (ToolTip = "Wheter to enable dynamic loading in the editor through loading regions"))
-	bool bDisableLoadingInEditor;
+	bool bEnableLoadingInEditor;
 
 	bool bDisableBugIt;
 	bool bDisablePIE;
diff --git a/Engine/Source/Editor/WorldPartitionEditor/Public/WorldPartitionEditorModule.h b/Engine/Source/Editor/WorldPartitionEditor/Public/WorldPartitionEditorModule.h
--- a/Engine/Source/Editor/WorldPartitionEditor/Public/WorldPartitionEditorModule.h
+++ b/Engine/Source/Editor/WorldPartitionEditor/Public/WorldPartitionEditorModule.h
@@ -48,10 +48,10 @@ public:
 	virtual int32 GetMinimapLowQualityWorldUnitsPerPixelThreshold() const override;
 
 	/**
-	 * Returns if loading in the editor is disabled or not.
+	 * Returns if loading in the editor is enabled or not.
 	 */
-	virtual bool GetDisableLoadingInEditor() const override;
-	virtual void SetDisableLoadingInEditor(bool bInDisableLoadingInEditor) override;
+	virtual bool GetEnableLoadingInEditor() const override;
+	virtual void SetEnableLoadingInEditor(bool bInEnableLoadingInEditor) override;
 
 	/**
 	 * Returns if pie is disabled or not.
diff --git a/Engine/Source/Runtime/Engine/Private/WorldPartition/WorldPartition.cpp b/Engine/Source/Runtime/Engine/Private/WorldPartition/WorldPartition.cpp
--- a/Engine/Source/Runtime/Engine/Private/WorldPartition/WorldPartition.cpp
+++ b/Engine/Source/Runtime/Engine/Private/WorldPartition/WorldPartition.cpp
@@ -563,7 +563,7 @@ void UWorldPartition::Initialize(UWorld* InWorld, const FTransform& InTransform)
 			PinnedActors = new FLoaderAdapterPinnedActors(OuterWorld);
 		
 			IWorldPartitionEditorModule& WorldPartitionEditorModule = FModuleManager::LoadModuleChecked<IWorldPartitionEditorModule>("WorldPartitionEditor");
-			ForceLoadedActors = WorldPartitionEditorModule.GetDisableLoadingInEditor() ? new FLoaderAdapterActorList(OuterWorld) : nullptr;
+			ForceLoadedActors = WorldPartitionEditorModule.GetEnableLoadingInEditor() ? nullptr : new FLoaderAdapterActorList(OuterWorld);
 		}
 	}
 
@@ -866,7 +866,7 @@ void UWorldPartition::OnPostBugItGoCalled(const FVector& Loc, const FRotator& Ro
 		const FBox LoadCellsBox(Loc - LoadExtent, Loc + LoadExtent);
 
 		IWorldPartitionEditorModule& WorldPartitionEditorModule = FModuleManager::LoadModuleChecked<IWorldPartitionEditorModule>("WorldPartitionEditor");
-		if (!WorldPartitionEditorModule.GetDisableLoadingInEditor())
+		if (WorldPartitionEditorModule.GetEnableLoadingInEditor())
 		{
 			UWorldPartitionEditorLoaderAdapter* EditorLoaderAdapter = CreateEditorLoaderAdapter<FLoaderAdapterShape>(World, LoadCellsBox, TEXT("BugItGo"));
 			EditorLoaderAdapter->GetLoaderAdapter()->Load();
@@ -1318,6 +1318,36 @@ void UWorldPartition::OnEnableStreamingChanged()
 	}
 }
 
+void UWorldPartition::OnEnableLoadingInEditorChanged()
+{
+	if (ForceLoadedActors)
+	{
+		delete ForceLoadedActors;
+		ForceLoadedActors = nullptr;
+	}
+
+	IWorldPartitionEditorModule& WorldPartitionEditorModule = FModuleManager::LoadModuleChecked<IWorldPartitionEditorModule>("WorldPartitionEditor");
+
+	if (!WorldPartitionEditorModule.GetEnableLoadingInEditor())
+	{
+		UWorld* OuterWorld = GetTypedOuter<UWorld>();
+		check(OuterWorld);
+
+		ForceLoadedActors = new FLoaderAdapterActorList(OuterWorld);
+
+		TArray<FGuid> ForceLoadedActorGuids;
+		for (FActorDescContainerCollection::TIterator<> ActorDescIterator(this); ActorDescIterator; ++ActorDescIterator)
+		{
+			ForceLoadedActorGuids.Add(ActorDescIterator->GetGuid());
+		}
+
+		if (ForceLoadedActorGuids.Num())
+		{
+			ForceLoadedActors->AddActors(ForceLoadedActorGuids);
+		}
+	}
+}
+
 void UWorldPartition::HashActorDesc(FWorldPartitionActorDesc* ActorDesc)
 {
 	check(ActorDesc);
diff --git a/Engine/Source/Runtime/Engine/Public/WorldPartition/IWorldPartitionEditorModule.h b/Engine/Source/Runtime/Engine/Public/WorldPartition/IWorldPartitionEditorModule.h
--- a/Engine/Source/Runtime/Engine/Public/WorldPartition/IWorldPartitionEditorModule.h
+++ b/Engine/Source/Runtime/Engine/Public/WorldPartition/IWorldPartitionEditorModule.h
@@ -36,8 +36,8 @@ public:
 	virtual int32 GetInstancedFoliageGridSize() const = 0;
 	virtual int32 GetMinimapLowQualityWorldUnitsPerPixelThreshold() const = 0;
 	
-	virtual bool GetDisableLoadingInEditor() const = 0;
-	virtual void SetDisableLoadingInEditor(bool bInDisableLoadingInEditor) = 0;
+	virtual bool GetEnableLoadingInEditor() const = 0;
+	virtual void SetEnableLoadingInEditor(bool bInEnableLoadingInEditor) = 0;
 
 	virtual bool GetDisablePIE() const = 0;
 	virtual void SetDisablePIE(bool bInDisablePIE) = 0;
diff --git a/Engine/Source/Runtime/Engine/Public/WorldPartition/WorldPartition.h b/Engine/Source/Runtime/Engine/Public/WorldPartition/WorldPartition.h
--- a/Engine/Source/Runtime/Engine/Public/WorldPartition/WorldPartition.h
+++ b/Engine/Source/Runtime/Engine/Public/WorldPartition/WorldPartition.h
@@ -130,6 +130,7 @@ public:
 	void SetCanBeUsedByLevelInstance(bool bInCanBeUsedByLevelInstance) {}
 
 	ENGINE_API void OnEnableStreamingChanged();
+	ENGINE_API void OnEnableLoadingInEditorChanged();
 
 private:
 	ENGINE_API void SavePerUserSettings();
@@ -536,4 +537,4 @@ private:
 
 	friend class AWorldPartitionReplay;
 	friend class UWorldPartitionSubsystem;
-};
+};
\ No newline at end of file
