World Partition Editor
* Upgrade overall UX (new icons, new menus, new toolbar, new layout)
* Now possible for licensees to extent options in our menus.
* Add same LevelViewport actions in the WP Editor (Build Hlods, build minimap, etc.) via new Build menu.

diff --git a/Engine/Source/Editor/EditorStyle/Private/StarshipStyle.cpp b/Engine/Source/Editor/EditorStyle/Private/StarshipStyle.cpp
--- a/Engine/Source/Editor/EditorStyle/Private/StarshipStyle.cpp
+++ b/Engine/Source/Editor/EditorStyle/Private/StarshipStyle.cpp
@@ -2002,7 +2002,6 @@ void FStarshipEditorStyle::FStyle::SetupLevelGeneralStyles()
 
 void FStarshipEditorStyle::FStyle::SetupWorldBrowserStyles()
 {
-
 	// World Browser
 	{
 		Set("WorldBrowser.AddLayer", new IMAGE_BRUSH("Icons/icon_levels_addlayer_16x", Icon16x16));
@@ -2039,9 +2038,19 @@ void FStarshipEditorStyle::FStyle::SetupWorldBrowserStyles()
 void FStarshipEditorStyle::FStyle::SetupWorldPartitionStyles()
 {
 	// World Partition
-	Set("WorldPartition.PartiallyLoadedCell", new IMAGE_BRUSH("WorldPartition/PartiallyLoadedCell", Icon32x32));
-	Set("WorldPartition.GridBackground", new BOX_BRUSH("WorldPartition/GridBackground", FMargin(0.f)));
 	Set("WorldPartition.SimulationViewPosition", new IMAGE_BRUSH("Icons/icon_levels_simulationviewpos_16x", Icon16x16));
+
+	Set("WorldPartition.FollowPlayerInPIE", new IMAGE_BRUSH_SVG("WorldPartition//Developer_20", Icon16x16));
+	Set("WorldPartition.BugItGoLoadRegion", new IMAGE_BRUSH_SVG("WorldPartition//Bug_20", Icon16x16)); 
+
+	Set("WorldPartition.ShowActors", new IMAGE_BRUSH_SVG("WorldPartition/Actor_20", Icon16x16));
+	Set("WorldPartition.ShowGrid", new IMAGE_BRUSH_SVG("Starship/EditorViewport/grid", Icon16x16));
+	Set("WorldPartition.ShowMiniMap", new IMAGE_BRUSH_SVG("Starship/AssetIcons/Texture2D_16", Icon16x16)); 
+	Set("WorldPartition.ShowCoords", new IMAGE_BRUSH_SVG("WorldPartition/Coordinate_20", Icon16x16));
+
+	Set("WorldPartition.FocusSelection", new IMAGE_BRUSH_SVG("WorldPartition/FramingSelection_20", Icon16x16));
+	Set("WorldPartition.FocusLoadedRegions", new IMAGE_BRUSH_SVG("WorldPartition/FrameRegions_20", Icon16x16));
+	Set("WorldPartition.FocusWorld", new IMAGE_BRUSH_SVG("WorldPartition/FrameWorld_20", Icon16x16));
 }
 
 void FStarshipEditorStyle::FStyle::SetupSequencerStyles()
diff --git a/Engine/Source/Editor/WorldPartitionEditor/Private/WorldPartition/SWorldPartitionEditorGrid2D.cpp b/Engine/Source/Editor/WorldPartitionEditor/Private/WorldPartition/SWorldPartitionEditorGrid2D.cpp
--- a/Engine/Source/Editor/WorldPartitionEditor/Private/WorldPartition/SWorldPartitionEditorGrid2D.cpp
+++ b/Engine/Source/Editor/WorldPartitionEditor/Private/WorldPartition/SWorldPartitionEditorGrid2D.cpp
@@ -14,14 +14,17 @@
 #include "Framework/Application/SlateApplication.h"
 #include "Layout/WidgetPath.h"
 #include "LevelEditorViewport.h"
+#include "LevelEditorActions.h"
 #include "LocationVolume.h"
 #include "Modules/ModuleManager.h"
 #include "Rendering/SlateRenderer.h"
 #include "ScopedTransaction.h"
 #include "SWorldPartitionViewportWidget.h"
+#include "SEditorViewportToolBarMenu.h"
 #include "TextureResource.h"
 #include "Widgets/Input/SButton.h"
 #include "Widgets/Input/SCheckBox.h"
+#include "Widgets/Layout/SUniformGridPanel.h"
 #include "Framework/MultiBox/MultiBoxBuilder.h"
 #include "WorldBrowserModule.h"
 #include "WorldPartition/LoaderAdapter/LoaderAdapterShape.h"
@@ -244,7 +247,7 @@ public:
 SWorldPartitionEditorGrid2D::FEditorCommands::FEditorCommands()
 	: TCommands<FEditorCommands>
 (
-	"WorldPartitionEditor",
+	"WorldPartition",
 	NSLOCTEXT("Contexts", "WorldPartition", "World Partition"),
 	NAME_None,
 	FAppStyle::GetAppStyleSetName()
@@ -253,6 +256,7 @@ SWorldPartitionEditorGrid2D::FEditorCommands::FEditorCommands()
 
 void SWorldPartitionEditorGrid2D::FEditorCommands::RegisterCommands()
 {
+	// Context Menu
 	UI_COMMAND(CreateRegionFromSelection, "Load Region From Selection", "Load region from selection.", EUserInterfaceActionType::Button, FInputChord());
 	UI_COMMAND(LoadSelectedRegions, "Load Selected Regions", "Load the selected regions.", EUserInterfaceActionType::Button, FInputChord());
 	UI_COMMAND(UnloadSelectedRegions, "Unload Selected Regions", "Unload the selected regions.", EUserInterfaceActionType::Button, FInputChord());
@@ -261,6 +265,159 @@ void SWorldPartitionEditorGrid2D::FEditorCommands::RegisterCommands()
 	UI_COMMAND(PlayFromHere, "Play From Here", "Play from here.", EUserInterfaceActionType::Button, FInputChord());
 	UI_COMMAND(LoadFromHere, "Load From Here", "Load from here.", EUserInterfaceActionType::Button, FInputChord());
 	UI_COMMAND(BugItHere, "Bug It Here", "Log BugItGo command of the selected location to console (also adds it to clipboard).", EUserInterfaceActionType::Button, FInputChord());
+
+	// Toolbar
+	UI_COMMAND(FollowPlayerInPIE, "Follow Player In PIE", "Follow Player In PIE.", EUserInterfaceActionType::ToggleButton, FInputChord());
+	UI_COMMAND(BugItGoLoadRegion, "BugItGo Load Region", "Using BugItGo command, it will create a loading region and zoom on it.", EUserInterfaceActionType::ToggleButton, FInputChord());
+
+	UI_COMMAND(ShowActors, "Actor(s)", "Show Actor(s).", EUserInterfaceActionType::ToggleButton, FInputChord());
+	UI_COMMAND(ShowGrid, "Grid", "Show Grid.", EUserInterfaceActionType::ToggleButton, FInputChord());
+	UI_COMMAND(ShowMiniMap, "Minimap", "Show the minimap texture.", EUserInterfaceActionType::ToggleButton, FInputChord());
+	UI_COMMAND(ShowCoords, "Coordinates", "Show grid cell coordinates, you might need to zoom closer to see them.", EUserInterfaceActionType::ToggleButton, FInputChord());
+
+	UI_COMMAND(FocusSelection, "Focus Selection", "Focus Selection.", EUserInterfaceActionType::ToggleButton, FInputChord());
+	UI_COMMAND(FocusLoadedRegions, "Focus Loaded Regions", "Focus Loaded Regions.", EUserInterfaceActionType::ToggleButton, FInputChord());
+	UI_COMMAND(FocusWorld, "Focus World", "Focus World.", EUserInterfaceActionType::ToggleButton, FInputChord());
+}
+
+void SWorldPartitionEditorGrid2D::SToolBar::Construct(const FArguments& InArgs)
+{
+	CommandList = InArgs._CommandList;
+
+	auto MakeQuickActionsToolBarWidget = [this]()
+	{
+		FSlimHorizontalToolBarBuilder ToolbarBuilder(CommandList, FMultiBoxCustomization::None);
+
+		FName ToolBarStyle = "EditorViewportToolBar";
+		ToolbarBuilder.SetStyle(&FAppStyle::Get(), ToolBarStyle);
+		ToolbarBuilder.SetLabelVisibility(EVisibility::Collapsed);
+
+		const FEditableTextBoxStyle& TextBoxStyle = FCoreStyle::Get().GetWidgetStyle<FEditableTextBoxStyle>("NormalEditableTextBox");
+
+		ToolbarBuilder.BeginSection("QuickActions");
+		{
+			ToolbarBuilder.BeginBlockGroup();
+
+			const FEditorCommands& Commands = FEditorCommands::Get();
+			ToolbarBuilder.AddToolBarButton(Commands.FocusSelection, NAME_None, TAttribute<FText>(), TAttribute<FText>(), TAttribute<FSlateIcon>());
+			ToolbarBuilder.AddToolBarButton(Commands.FocusLoadedRegions, NAME_None, TAttribute<FText>(), TAttribute<FText>(), TAttribute<FSlateIcon>());
+			ToolbarBuilder.AddToolBarButton(Commands.FocusWorld, NAME_None, TAttribute<FText>(), TAttribute<FText>(), TAttribute<FSlateIcon>());
+
+			ToolbarBuilder.EndBlockGroup();
+		}
+		ToolbarBuilder.EndSection();
+
+		return ToolbarBuilder.MakeWidget();
+	};
+
+
+	const FMargin ToolbarSlotPadding(4.0f, 1.0f);
+	const FMargin ToolbarButtonPadding(4.0f, 0.0f);
+
+	TSharedPtr<SHorizontalBox> MainBoxPtr;
+
+	ChildSlot
+	[
+		SNew(SBorder)
+		.BorderImage(FAppStyle::Get().GetBrush("EditorViewportToolBar.Background"))
+		.Cursor(EMouseCursor::Default)
+		[
+			SAssignNew(MainBoxPtr, SHorizontalBox)
+		]
+	];
+
+	// Options menu
+	MainBoxPtr->AddSlot()
+		.AutoWidth()
+		.Padding(ToolbarSlotPadding)
+		[
+			SNew(SEditorViewportToolbarMenu)
+			.ParentToolBar(SharedThis(this))
+			.Cursor(EMouseCursor::Default)
+			.Image("EditorViewportToolBar.OptionsDropdown")
+			.OnGetMenuContent(this, &SToolBar::GenerateOptionsMenu)
+		];
+
+	// Show menu
+	MainBoxPtr->AddSlot()
+		.AutoWidth()
+		.Padding(ToolbarSlotPadding)
+		[
+			SNew(SEditorViewportToolbarMenu)
+			.Label(LOCTEXT("ShowMenuTitle", "Show"))
+			.Cursor(EMouseCursor::Default)
+			.ParentToolBar(SharedThis(this))
+			.OnGetMenuContent(this, &SToolBar::GenerateShowMenu)
+		];
+
+	// Build menu
+	MainBoxPtr->AddSlot()
+		.AutoWidth()
+		.Padding(ToolbarSlotPadding)
+		[
+			SNew(SEditorViewportToolbarMenu)
+			.Label(LOCTEXT("BuildMenuTitle", "Build"))
+			.Cursor(EMouseCursor::Default)
+			.ParentToolBar(SharedThis(this))
+			.OnGetMenuContent(this, &SToolBar::GenerateBuildMenu)
+		];
+
+	// Quick Actions Toolbar
+	MainBoxPtr->AddSlot()
+		.Padding(ToolbarSlotPadding)
+		.HAlign(HAlign_Right)
+		[
+			MakeQuickActionsToolBarWidget()
+		];
+}
+
+TSharedRef<SWidget> SWorldPartitionEditorGrid2D::SToolBar::GenerateOptionsMenu() const
+{
+	static const FName MenuName(TEXT("WorldPartition.OptionsMenu"));
+	static const FName SectionName(TEXT("Options"));
+
+	UToolMenu* ShowMenu = UToolMenus::Get()->RegisterMenu(MenuName);
+
+	FToolMenuSection& Section = ShowMenu->FindOrAddSection(SectionName);
+	const FEditorCommands& Commands = FEditorCommands::Get();
+	Section.AddMenuEntry(Commands.FollowPlayerInPIE);
+	Section.AddMenuEntry(Commands.BugItGoLoadRegion);
+
+	return UToolMenus::Get()->GenerateWidget(MenuName, FToolMenuContext(CommandList));
+}
+
+TSharedRef<SWidget> SWorldPartitionEditorGrid2D::SToolBar::GenerateShowMenu() const
+{
+	static const FName MenuName(TEXT("WorldPartition.ShowMenu"));
+	static const FName SectionName(TEXT("Show"));
+
+	UToolMenu* ShowMenu = UToolMenus::Get()->RegisterMenu(MenuName);
+
+	FToolMenuSection& Section = ShowMenu->FindOrAddSection(SectionName);
+	const FEditorCommands& Commands = FEditorCommands::Get();
+	Section.AddMenuEntry(Commands.ShowActors);
+	Section.AddMenuEntry(Commands.ShowGrid);
+	Section.AddMenuEntry(Commands.ShowMiniMap);
+	Section.AddMenuEntry(Commands.ShowCoords);
+
+	return UToolMenus::Get()->GenerateWidget(MenuName, FToolMenuContext(CommandList));
+}
+
+TSharedRef<SWidget> SWorldPartitionEditorGrid2D::SToolBar::GenerateBuildMenu() const
+{
+	static const FName MenuName(TEXT("WorldPartition.BuildMenu"));
+	static const FName SectionName(TEXT("Build"));
+
+	UToolMenu* BuildMenu = UToolMenus::Get()->RegisterMenu(MenuName);
+
+	FToolMenuSection& Section = BuildMenu->FindOrAddSection(SectionName);
+	FLevelEditorModule& LevelEditorModule = FModuleManager::GetModuleChecked<FLevelEditorModule>(TEXT("LevelEditor"));
+	const FLevelEditorCommands& LevelEditorCommands = LevelEditorModule.GetLevelEditorCommands();
+	Section.AddMenuEntry(LevelEditorCommands.BuildHLODs);
+	Section.AddMenuEntry(LevelEditorCommands.BuildMinimap);
+	Section.AddMenuEntry(LevelEditorCommands.BuildLandscapeSplineMeshes);
+
+	return UToolMenus::Get()->GenerateWidget(MenuName, FToolMenuContext(LevelEditorModule.GetGlobalLevelEditorActions()));
 }
 
 SWorldPartitionEditorGrid2D::SWorldPartitionEditorGrid2D()
@@ -274,6 +431,7 @@ SWorldPartitionEditorGrid2D::SWorldPartitionEditorGrid2D()
 	, bIsMeasuring(false)
 	, bShowActors(false)
 	, bShowGrid(true)
+	, bShowMiniMap(true)
 	, bFollowPlayerInPIE(false)
 	, SelectBox(ForceInit)
 	, SelectBoxGridSnapped(ForceInit)
@@ -294,161 +452,82 @@ void SWorldPartitionEditorGrid2D::Construct(const FArguments& InArgs)
 {
 	SWorldPartitionEditorGrid::Construct(SWorldPartitionEditorGrid::FArguments().InWorld(InArgs._InWorld));
 
-	if (GetWorldPartition())
-	{
-		UpdateWorldMiniMapDetails();
-		bShowActors = !WorldMiniMapBrush.HasUObject();
-	}
-
 	// Defaults
 	Trans = FVector2D(0, 0);
 	Scale = 0.00133333332;
 	TotalMouseDelta = 0;
+	SmallLayoutFont = FCoreStyle::GetDefaultFontStyle("Regular", 10);
 
-	auto MakeToolBarWidget = [this]()
+	if (GetWorldPartition())
 	{
-		FSlimHorizontalToolBarBuilder ToolbarBuilder(CommandList, FMultiBoxCustomization::None);
-
-		FName ToolBarStyle = "EditorViewportToolBar";
-		ToolbarBuilder.SetStyle(&FAppStyle::Get(), ToolBarStyle);
-		ToolbarBuilder.SetLabelVisibility(EVisibility::Collapsed);
-
-		const FEditableTextBoxStyle& TextBoxStyle = FCoreStyle::Get().GetWidgetStyle<FEditableTextBoxStyle>("NormalEditableTextBox");
-
-		ToolbarBuilder.BeginSection("Toggles");
-		{
-			ToolbarBuilder.AddWidget(SNew(SCheckBox)
-				.ForegroundColor(FSlateColor::UseForeground())
-				.BorderBackgroundColor(TextBoxStyle.BackgroundColor)
-				.IsChecked(bShowActors ? ECheckBoxState::Checked : ECheckBoxState::Unchecked)
-				.OnCheckStateChanged(FOnCheckStateChanged::CreateLambda([this](ECheckBoxState State) { bShowActors = !bShowActors; }))
-				[
-					SNew(SBox)
-					.VAlign(VAlign_Center)
-					.HAlign(HAlign_Center)
-					.Padding(FMargin(4.0f, 2.0f))
-					[
-						SNew(STextBlock)
-						.Text(LOCTEXT("ShowActors", "Show Actors"))
-					]
-				]);
-
-			ToolbarBuilder.AddWidget(SNew(SCheckBox)
-				.ForegroundColor(FSlateColor::UseForeground())
-				.BorderBackgroundColor(TextBoxStyle.BackgroundColor)
-				.IsChecked(bShowGrid ? ECheckBoxState::Checked : ECheckBoxState::Unchecked)
-				.OnCheckStateChanged(FOnCheckStateChanged::CreateLambda([this](ECheckBoxState State) { bShowGrid = !bShowGrid; }))
-				[
-					SNew(SBox)
-					.VAlign(VAlign_Center)
-					.HAlign(HAlign_Center)
-					.Padding(FMargin(4.0f, 2.0f))
-					[
-						SNew(STextBlock)
-						.Text(LOCTEXT("ShowGrid", "Show Grid"))
-					]
-				]);
-
-			ToolbarBuilder.AddWidget(SNew(SCheckBox)
-				.ForegroundColor(FSlateColor::UseForeground())
-				.IsChecked(bFollowPlayerInPIE ? ECheckBoxState::Checked : ECheckBoxState::Unchecked)
-				.Visibility_Lambda([this]() { return GetDefault<UWorldPartitionEditorSettings>()->bDisablePIE ? EVisibility::Hidden : EVisibility::Visible; })
-				.OnCheckStateChanged(FOnCheckStateChanged::CreateLambda([this](ECheckBoxState State) { bFollowPlayerInPIE = !bFollowPlayerInPIE; }))
-				[
-					SNew(SBox)
-					.VAlign(VAlign_Center)
-					.HAlign(HAlign_Center)
-					.Padding(FMargin(4.0f, 2.0f))
-					[
-						SNew(STextBlock)
-						.Text(LOCTEXT("FollowPlayerInPIE", "Follow Player in PIE"))
-					]
-				]);
-
-			ToolbarBuilder.AddWidget(SNew(SCheckBox)
-				.ForegroundColor(FSlateColor::UseForeground())
-				.IsChecked(GetMutableDefault<UWorldPartitionEditorPerProjectUserSettings>()->GetBugItGoLoadRegion() ? ECheckBoxState::Checked : ECheckBoxState::Unchecked)
-				.Visibility_Lambda([this]() { return GetDefault<UWorldPartitionEditorSettings>()->bEnableLoadingInEditor ? EVisibility::Visible : EVisibility::Hidden; })
-				.OnCheckStateChanged(FOnCheckStateChanged::CreateLambda([=](ECheckBoxState State) { GetMutableDefault<UWorldPartitionEditorPerProjectUserSettings>()->SetBugItGoLoadRegion(State == ECheckBoxState::Checked); }))
-				[
-					SNew(SBox)
-					.VAlign(VAlign_Center)
-					.HAlign(HAlign_Center)
-					.Padding(FMargin(4.0f, 2.0f))
-					[
-						SNew(STextBlock)
-						.Text(LOCTEXT("BugItGoLoadRegion", "BugItGo Load Region"))
-					]
-				]);
-
-			ToolbarBuilder.AddWidget(SNew(SCheckBox)
-				.ForegroundColor(FSlateColor::UseForeground())
-				.IsChecked(GetMutableDefault<UWorldPartitionEditorPerProjectUserSettings>()->GetShowCellCoords() ? ECheckBoxState::Checked : ECheckBoxState::Unchecked)
-				.Visibility_Lambda([this]() { return (GetWorldPartition() && GetWorldPartition()->IsStreamingEnabled()) ? EVisibility::Visible : EVisibility::Hidden; })
-				.OnCheckStateChanged(FOnCheckStateChanged::CreateLambda([=](ECheckBoxState State) { GetMutableDefault<UWorldPartitionEditorPerProjectUserSettings>()->SetShowCellCoords(State == ECheckBoxState::Checked); }))
-				[
-					SNew(SBox)
-					.VAlign(VAlign_Center)
-					.HAlign(HAlign_Center)
-					.Padding(FMargin(4.0f, 2.0f))
-					[
-						SNew(STextBlock)
-						.Text(LOCTEXT("ShowCellCoords", "Show Cell Coords"))
-					]
-				]);
-
-			ToolbarBuilder.AddWidget(SNew(SButton)
-				.ForegroundColor(FSlateColor::UseForeground())
-				.Text_Lambda([this]()
-				{
-					if (UWorldPartitionSubsystem* WorldPartitionSubsystem = UWorld::GetSubsystem<UWorldPartitionSubsystem>(GetWorld()))
-					{
-						return (GEditor->GetSelectedActors()->Num() || WorldPartitionSubsystem->SelectedActorDescs.Num()) ? LOCTEXT("FocusSelection", "Focus Selection") : LOCTEXT("FocusWorld", "Focus World");
-					}
-					return LOCTEXT("FocusWorld", "Focus World");
-				})
-				.OnClicked(this, &SWorldPartitionEditorGrid2D::FocusSelection)
-				.IsEnabled_Lambda([this]() { return IsInteractive(); }));
-
-			ToolbarBuilder.AddWidget(SNew(SButton)
-				.ForegroundColor(FSlateColor::UseForeground())
-				.Text(LOCTEXT("FocusLoadedRegions", "Focus Loaded Regions"))
-				.OnClicked(this, &SWorldPartitionEditorGrid2D::FocusLoadedRegions)
-				.IsEnabled_Lambda([this]() { return IsInteractive() && GetWorldPartition() && GetWorldPartition()->HasLoadedUserCreatedRegions(); })
-				.Visibility_Lambda([this]() { return GetDefault<UWorldPartitionEditorSettings>()->bEnableLoadingInEditor ? EVisibility::Visible : EVisibility::Hidden; }));
-		}
-
-		return ToolbarBuilder.MakeWidget();
-	};
+		UpdateWorldMiniMapDetails();
+		bShowActors = !WorldMiniMapBrush.HasUObject();
+	}
 
 	// UI
+	const FMargin ToolbarSlotPadding(4.0f, 1.0f);
+	const float BottomPartMinDesiredHeight(25.0f); // Need this because the left part is a widget bigger than other text parts and his visibility could be turn off resulting in moving the text on screen (unwanted behavior).
+
 	ChildSlot
 	[
 		SNew(SOverlay)
 
-		// Top status bar
-		+SOverlay::Slot()
+		// Toolbar
+		+ SOverlay::Slot()
 		.VAlign(VAlign_Top)
+		.Padding(ToolbarSlotPadding)
 		[
-			MakeToolBarWidget()
+			SNew(SToolBar)
+			.CommandList(CommandList)
 		]
+
+		// Bottom section
 		+SOverlay::Slot()
 		.VAlign(VAlign_Bottom)
-		.HAlign(HAlign_Left)
-		.Padding(10.f, 0.f, 0.f, 10.f)
+
+		.Padding(10.f, 0.f, 0.f, 1.f)
 		[
-			SAssignNew(ViewportWidget, SWorldPartitionViewportWidget)
-			.Clickable(false)
-			.Visibility_Lambda([this]() { return ViewportWidget->GetVisibility(GetWorld()); })
+			SNew(SUniformGridPanel)
+			.SlotPadding(2)
+			 + SUniformGridPanel::Slot(0, 0)
+			[
+				SNew(SBox)
+				.HAlign(HAlign_Left)
+				.VAlign(VAlign_Center)
+				[
+					SAssignNew(ViewportWidget, SWorldPartitionViewportWidget)
+					.Clickable(false)
+					.Visibility_Lambda([this]() { return ViewportWidget->GetVisibility(GetWorld()); })
+				]
+			]
+			+ SUniformGridPanel::Slot(1, 0)
+			[
+				SNew(SBox)
+				.HAlign(HAlign_Center)
+				.VAlign(VAlign_Center)
+				.MinDesiredHeight(BottomPartMinDesiredHeight)
+				[
+					SAssignNew(TextWorldBoundsInKMWidget, STextBlock)
+				]
+			]
+			+ SUniformGridPanel::Slot(2, 0)
+			[
+				SNew(SBox)
+				.HAlign(HAlign_Right)
+				.VAlign(VAlign_Center)
+				.MinDesiredHeight(BottomPartMinDesiredHeight)
+				[
+					SAssignNew(TextRulerWidget, STextBlock)
+				]
+			]
 		]
 	];
 
-	SmallLayoutFont = FCoreStyle::GetDefaultFontStyle("Regular", 10);
-
-	// Bind commands
-	const FEditorCommands& Commands = FEditorCommands::Get();
-	FUICommandList& ActionList = *CommandList;
+	BindCommands();
+}
 
+void SWorldPartitionEditorGrid2D::BindCommands()
+{
 	auto CanCreateRegionFromSelection = [this]()
 	{
 		return !!SelectBoxGridSnapped.IsValid;
@@ -490,14 +569,42 @@ void SWorldPartitionEditorGrid2D::Construct(const FArguments& InArgs)
 		return GetWorldPartition()->IsStreamingEnabled();
 	};
 
-	ActionList.MapAction(Commands.CreateRegionFromSelection, FExecuteAction::CreateSP(this, &SWorldPartitionEditorGrid2D::CreateRegionFromSelection), FCanExecuteAction::CreateLambda(CanCreateRegionFromSelection));
-	ActionList.MapAction(Commands.LoadSelectedRegions, FExecuteAction::CreateSP(this, &SWorldPartitionEditorGrid2D::LoadSelectedRegions), FCanExecuteAction::CreateLambda(CanLoadSelectedRegions));
-	ActionList.MapAction(Commands.UnloadSelectedRegions, FExecuteAction::CreateSP(this, &SWorldPartitionEditorGrid2D::UnloadSelectedRegions), FCanExecuteAction::CreateLambda(CanUnloadSelectedRegions));
-	ActionList.MapAction(Commands.ConvertSelectedRegionsToActors, FExecuteAction::CreateSP(this, &SWorldPartitionEditorGrid2D::ConvertSelectedRegionsToActors), FCanExecuteAction::CreateLambda(CanConvertSelectedRegionsToActors));
-	ActionList.MapAction(Commands.MoveCameraHere, FExecuteAction::CreateSP(this, &SWorldPartitionEditorGrid2D::MoveCameraHere));
-	ActionList.MapAction(Commands.PlayFromHere, FExecuteAction::CreateSP(this, &SWorldPartitionEditorGrid2D::PlayFromHere));
-	ActionList.MapAction(Commands.LoadFromHere, FExecuteAction::CreateSP(this, &SWorldPartitionEditorGrid2D::LoadFromHere), FCanExecuteAction::CreateLambda(CanLoadFromHere));
-	ActionList.MapAction(Commands.BugItHere, FExecuteAction::CreateSP(this, &SWorldPartitionEditorGrid2D::BugItHere));
+	auto CanFocusSelection = [this]()
+	{
+		if (UWorldPartitionSubsystem* WorldPartitionSubsystem = UWorld::GetSubsystem<UWorldPartitionSubsystem>(GetWorld()))
+		{
+			return IsInteractive() && (GEditor->GetSelectedActors()->Num() || WorldPartitionSubsystem->SelectedActorDescs.Num());
+		}
+
+		return false;
+	};
+
+
+	const FEditorCommands& Commands = FEditorCommands::Get();
+
+	CommandList->MapAction(Commands.CreateRegionFromSelection, FExecuteAction::CreateSP(this, &SWorldPartitionEditorGrid2D::CreateRegionFromSelection), FCanExecuteAction::CreateLambda(CanCreateRegionFromSelection));
+	CommandList->MapAction(Commands.LoadSelectedRegions, FExecuteAction::CreateSP(this, &SWorldPartitionEditorGrid2D::LoadSelectedRegions), FCanExecuteAction::CreateLambda(CanLoadSelectedRegions));
+	CommandList->MapAction(Commands.UnloadSelectedRegions, FExecuteAction::CreateSP(this, &SWorldPartitionEditorGrid2D::UnloadSelectedRegions), FCanExecuteAction::CreateLambda(CanUnloadSelectedRegions));
+	CommandList->MapAction(Commands.ConvertSelectedRegionsToActors, FExecuteAction::CreateSP(this, &SWorldPartitionEditorGrid2D::ConvertSelectedRegionsToActors), FCanExecuteAction::CreateLambda(CanConvertSelectedRegionsToActors));
+	CommandList->MapAction(Commands.MoveCameraHere, FExecuteAction::CreateSP(this, &SWorldPartitionEditorGrid2D::MoveCameraHere));
+	CommandList->MapAction(Commands.PlayFromHere, FExecuteAction::CreateSP(this, &SWorldPartitionEditorGrid2D::PlayFromHere));
+	CommandList->MapAction(Commands.LoadFromHere, FExecuteAction::CreateSP(this, &SWorldPartitionEditorGrid2D::LoadFromHere), FCanExecuteAction::CreateLambda(CanLoadFromHere));
+	CommandList->MapAction(Commands.BugItHere, FExecuteAction::CreateSP(this, &SWorldPartitionEditorGrid2D::BugItHere));
+
+	// Options
+	CommandList->MapAction(Commands.FollowPlayerInPIE, FExecuteAction::CreateLambda([this]() { bFollowPlayerInPIE = !bFollowPlayerInPIE; }), FCanExecuteAction(), FIsActionChecked::CreateLambda([this]() { return bFollowPlayerInPIE; }), FIsActionButtonVisible::CreateLambda([this]() { return !GetDefault<UWorldPartitionEditorSettings>()->bDisablePIE; }));
+	CommandList->MapAction(Commands.BugItGoLoadRegion, FExecuteAction::CreateLambda([this]() { GetMutableDefault<UWorldPartitionEditorPerProjectUserSettings>()->SetBugItGoLoadRegion(!GetMutableDefault<UWorldPartitionEditorPerProjectUserSettings>()->GetBugItGoLoadRegion()); }), FCanExecuteAction(), FIsActionChecked::CreateLambda([this]() { return GetMutableDefault<UWorldPartitionEditorPerProjectUserSettings>()->GetBugItGoLoadRegion(); }), FIsActionButtonVisible::CreateLambda([this]() { return GetDefault<UWorldPartitionEditorSettings>()->bEnableLoadingInEditor; }));
+
+	// Show toggles
+	CommandList->MapAction(Commands.ShowActors, FExecuteAction::CreateLambda([this]() { bShowActors = !bShowActors; }), FCanExecuteAction(), FIsActionChecked::CreateLambda([this]() { return bShowActors; }));
+	CommandList->MapAction(Commands.ShowGrid, FExecuteAction::CreateLambda([this]() { bShowGrid = !bShowGrid; }), FCanExecuteAction(), FIsActionChecked::CreateLambda([this]() { return bShowGrid; }));
+	CommandList->MapAction(Commands.ShowMiniMap, FExecuteAction::CreateLambda([this]() { bShowMiniMap = !bShowMiniMap; }), FCanExecuteAction(), FIsActionChecked::CreateLambda([this]() { return bShowMiniMap; }));
+	CommandList->MapAction(Commands.ShowCoords, FExecuteAction::CreateLambda([this]() { GetMutableDefault<UWorldPartitionEditorPerProjectUserSettings>()->SetShowCellCoords(!GetMutableDefault<UWorldPartitionEditorPerProjectUserSettings>()->GetShowCellCoords()); }), FCanExecuteAction(), FIsActionChecked::CreateLambda([this]() { return GetMutableDefault<UWorldPartitionEditorPerProjectUserSettings>()->GetShowCellCoords(); }), FIsActionButtonVisible::CreateLambda([this]() { return (GetWorldPartition() && GetWorldPartition()->IsStreamingEnabled()); }));
+
+	// Buttons
+	CommandList->MapAction(Commands.FocusSelection, FExecuteAction::CreateSP(this, &SWorldPartitionEditorGrid2D::FocusSelection), FCanExecuteAction::CreateLambda(CanFocusSelection));
+	CommandList->MapAction(Commands.FocusLoadedRegions, FExecuteAction::CreateSP(this, &SWorldPartitionEditorGrid2D::FocusLoadedRegions), FCanExecuteAction::CreateLambda([this]() { return IsInteractive() && GetWorldPartition() && GetWorldPartition()->HasLoadedUserCreatedRegions(); }), FIsActionChecked(), FIsActionButtonVisible::CreateLambda([this]() { return GetDefault<UWorldPartitionEditorSettings>()->bEnableLoadingInEditor; }));
+	CommandList->MapAction(Commands.FocusWorld, FExecuteAction::CreateSP(this, &SWorldPartitionEditorGrid2D::FocusWorld), FCanExecuteAction::CreateLambda([this]() { return IsInteractive(); }));
 }
 
 void SWorldPartitionEditorGrid2D::UpdateWorldMiniMapDetails()
@@ -540,7 +647,7 @@ void SWorldPartitionEditorGrid2D::LoadSelectedRegions()
 
 	SelectBox.Init();
 	SelectBoxGridSnapped.Init();
-		
+
 	GEditor->RedrawLevelEditingViewports();
 	Refresh();
 }
@@ -1492,10 +1599,13 @@ uint32 SWorldPartitionEditorGrid2D::PaintTextInfo(const FGeometry& AllottedGeome
 	LinePoints.Add(FVector2D::ZeroVector);
 	LinePoints.Add(FVector2D::ZeroVector + FVector2D(ScaleRulerLength, 0.f));
 	
-	FSlateDrawElement::MakeLines( 
+	int InX = AllottedGeometry.GetLocalSize().X - ScaleRulerLength;
+	int InY = AllottedGeometry.GetLocalSize().Y - 8;
+
+	FSlateDrawElement::MakeLines(
 		OutDrawElements,
 		LayerId,
-		AllottedGeometry.ToOffsetPaintGeometry(FVector2D(10, 40)),
+		AllottedGeometry.ToOffsetPaintGeometry(FVector2D(InX, InY)),
 		LinePoints,
 		ESlateDrawEffect::None,
 		FLinearColor::White);
@@ -1504,38 +1614,19 @@ uint32 SWorldPartitionEditorGrid2D::PaintTextInfo(const FGeometry& AllottedGeome
 	const int32 UnitsInMeter = 100;
 	const int32 UnitsInKilometer = UnitsInMeter*1000;
 	
-	FString RulerText;
-	if (UnitsInRuler >= UnitsInKilometer) // in kilometers
-	{
-		RulerText = FString::Printf(TEXT("%.2f km"), UnitsInRuler/UnitsInKilometer);
-	}
-	else // in meters
+	FString RulerText = (UnitsInRuler >= UnitsInKilometer) ? FString::Printf(TEXT("%.2f km"), UnitsInRuler / UnitsInKilometer) : FString::Printf(TEXT("%.2f m"), UnitsInRuler / UnitsInMeter);
+	if (TextRulerWidget.IsValid())
 	{
-		RulerText = FString::Printf(TEXT("%.2f m"), UnitsInRuler/UnitsInMeter);
+		TextRulerWidget->SetText(FText::FromString(RulerText));
 	}
-	
-	FSlateDrawElement::MakeText(
-		OutDrawElements,
-		LayerId,
-		AllottedGeometry.ToOffsetPaintGeometry(FVector2D(10, 27)),
-		RulerText,
-		FAppStyle::GetFontStyle("NormalFont"),
-		ESlateDrawEffect::None,
-		FLinearColor::White);
 
 	// Show world bounds
 	const FBox WorldBounds = GetWorldPartition()->GetRuntimeWorldBounds();
 	const FVector WorldBoundsExtentInKM = (WorldBounds.GetExtent() * 2.0f) / 100000.0f;
-	RulerText = FString::Printf(TEXT("%.2fx%.2fx%.2f km"), WorldBoundsExtentInKM.X, WorldBoundsExtentInKM.Y, WorldBoundsExtentInKM.Z);
-
-	FSlateDrawElement::MakeText(
-		OutDrawElements,
-		LayerId,
-		AllottedGeometry.ToOffsetPaintGeometry(FVector2D(10, 67)),
-		RulerText,
-		FAppStyle::GetFontStyle("NormalFont"),
-		ESlateDrawEffect::None,
-		FLinearColor::White);
+	if (TextWorldBoundsInKMWidget.IsValid())
+	{
+		TextWorldBoundsInKMWidget->SetText(FText::FromString(FString::Printf(TEXT("%.2fx%.2fx%.2f km"), WorldBoundsExtentInKM.X, WorldBoundsExtentInKM.Y, WorldBoundsExtentInKM.Z)));
+	}
 
 	// Show profiling
 	if (GShowEditorProfilingStats)
@@ -1976,7 +2067,7 @@ int32 SWorldPartitionEditorGrid2D::DrawTextLabel(FSlateWindowElementList& OutDra
 	return LayerId;
 }
 
-FReply SWorldPartitionEditorGrid2D::FocusSelection()
+void SWorldPartitionEditorGrid2D::FocusSelection()
 {
 	FBox SelectionBox(ForceInit);
 
@@ -1994,28 +2085,13 @@ FReply SWorldPartitionEditorGrid2D::FocusSelection()
 		SelectionBox += SelectedActorDesc->GetEditorBounds();
 	}
 
-	if (!SelectionBox.IsValid)
+	if (SelectionBox.IsValid)
 	{
-		// Override the minimap bounds if world partition minimap volumes exists
-		for (TActorIterator<AWorldPartitionMiniMapVolume> It(GetWorld()); It; ++It)
-		{
-			if (AWorldPartitionMiniMapVolume* WorldPartitionMiniMapVolume = *It)
-			{
-				SelectionBox += WorldPartitionMiniMapVolume->GetBounds().GetBox();
-			}
-		}
-
-		if (!SelectionBox.IsValid)
-		{
-			SelectionBox = GetWorldPartition()->GetRuntimeWorldBounds();
-		}
+		FocusBox(SelectionBox);
 	}
-
-	FocusBox(SelectionBox);
-	return FReply::Handled();
 }
 
-FReply SWorldPartitionEditorGrid2D::FocusLoadedRegions()
+void SWorldPartitionEditorGrid2D::FocusLoadedRegions()
 {
 	FBox SelectionBox(ForceInit);
 
@@ -2049,7 +2125,27 @@ FReply SWorldPartitionEditorGrid2D::FocusLoadedRegions()
 	}
 
 	FocusBox(SelectionBox);
-	return FReply::Handled();
+}
+
+void SWorldPartitionEditorGrid2D::FocusWorld()
+{
+	FBox SelectionBox(ForceInit);
+
+	// Override the minimap bounds if world partition minimap volumes exists
+	for (TActorIterator<AWorldPartitionMiniMapVolume> It(GetWorld()); It; ++It)
+	{
+		if (AWorldPartitionMiniMapVolume* WorldPartitionMiniMapVolume = *It)
+		{
+			SelectionBox += WorldPartitionMiniMapVolume->GetBounds().GetBox();
+		}
+	}
+
+	if (!SelectionBox.IsValid)
+	{
+		SelectionBox = GetWorldPartition()->GetRuntimeWorldBounds();
+	}
+
+	FocusBox(SelectionBox);
 }
 
 void SWorldPartitionEditorGrid2D::FocusBox(const FBox& Box) const
diff --git a/Engine/Source/Editor/WorldPartitionEditor/Private/WorldPartition/SWorldPartitionEditorGridSpatialHash.cpp b/Engine/Source/Editor/WorldPartitionEditor/Private/WorldPartition/SWorldPartitionEditorGridSpatialHash.cpp
--- a/Engine/Source/Editor/WorldPartitionEditor/Private/WorldPartition/SWorldPartitionEditorGridSpatialHash.cpp
+++ b/Engine/Source/Editor/WorldPartitionEditor/Private/WorldPartition/SWorldPartitionEditorGridSpatialHash.cpp
@@ -75,6 +75,7 @@ int32 SWorldPartitionEditorGridSpatialHash::PaintGrid(const FGeometry& AllottedG
 	);
 
 	// Paint minimap
+	if (bShowMiniMap)
 	{
 		const FBox VisibleMinimapRectWorld(
 			FVector(
diff --git a/Engine/Source/Editor/WorldPartitionEditor/Public/WorldPartition/SWorldPartitionEditorGrid2D.h b/Engine/Source/Editor/WorldPartitionEditor/Public/WorldPartition/SWorldPartitionEditorGrid2D.h
--- a/Engine/Source/Editor/WorldPartitionEditor/Public/WorldPartition/SWorldPartitionEditorGrid2D.h
+++ b/Engine/Source/Editor/WorldPartitionEditor/Public/WorldPartition/SWorldPartitionEditorGrid2D.h
@@ -9,6 +9,8 @@
 #include "Layout/ArrangedChildren.h"
 #include "SWorldPartitionEditorGrid.h"
 #include "SWorldPartitionViewportWidget.h"
+#include "Widgets/Text/STextBlock.h"
+#include "SViewportToolBar.h"
 #include "WorldPartition/WorldPartitionActorLoaderInterface.h"
 
 class SWorldPartitionEditorGrid2D : public SWorldPartitionEditorGrid
@@ -19,6 +21,7 @@ protected:
 	public:
 		FEditorCommands();
 	
+		// Context Menu
 		TSharedPtr<FUICommandInfo> CreateRegionFromSelection;
 		TSharedPtr<FUICommandInfo> LoadSelectedRegions;
 		TSharedPtr<FUICommandInfo> UnloadSelectedRegions;
@@ -29,18 +32,56 @@ protected:
 		TSharedPtr<FUICommandInfo> LoadFromHere;
 		TSharedPtr<FUICommandInfo> BugItHere;
 
+		// Toolbar
+		// Options
+		TSharedPtr<FUICommandInfo> FollowPlayerInPIE;
+		TSharedPtr<FUICommandInfo> BugItGoLoadRegion;
+
+		// Show toggles
+		TSharedPtr<FUICommandInfo> ShowActors;
+		TSharedPtr<FUICommandInfo> ShowGrid;
+		TSharedPtr<FUICommandInfo> ShowMiniMap;
+		TSharedPtr<FUICommandInfo> ShowCoords;
+
+		// Quick Actions
+		TSharedPtr<FUICommandInfo> FocusSelection;
+		TSharedPtr<FUICommandInfo> FocusLoadedRegions;
+		TSharedPtr<FUICommandInfo> FocusWorld;
+		
+
 		/**
 		 * Initialize commands
 		 */
 		virtual void RegisterCommands() override;
 	};
 
+
+	// In-viewport toolbar widget used in the world partition editor
+	class SToolBar : public SViewportToolBar
+	{
+		public:
+			SLATE_BEGIN_ARGS(SToolBar) {}
+				SLATE_ARGUMENT(TSharedPtr<FUICommandList>, CommandList)
+			SLATE_END_ARGS()
+
+			void Construct(const FArguments& InArgs);
+
+	private:
+		TSharedRef<SWidget> GenerateOptionsMenu() const;
+		TSharedRef<SWidget> GenerateShowMenu() const;
+		TSharedRef<SWidget> GenerateBuildMenu() const;
+		
+		TSharedPtr<FUICommandList> CommandList;
+	};
+
+
 public:
 	SWorldPartitionEditorGrid2D();
 	~SWorldPartitionEditorGrid2D();
 
 protected:
 	void Construct(const FArguments& InArgs);
+	void BindCommands();
 
 	void CreateRegionFromSelection();
 	void LoadSelectedRegions();
@@ -76,9 +117,10 @@ protected:
 	virtual int32 PaintMeasureTool(const FGeometry& AllottedGeometry, const FSlateRect& MyCullingRect, FSlateWindowElementList& OutDrawElements, int32 LayerId) const;
 
 	int32 DrawTextLabel(FSlateWindowElementList& OutDrawElements, int32 LayerId, const FGeometry& AllottedGeometry, const FString& Label, const FVector2D& Pos, const FLinearColor& Color, const FSlateFontInfo& Font) const;
-
-	virtual FReply FocusSelection();
-	virtual FReply FocusLoadedRegions();
+	
+	void FocusSelection();
+	void FocusLoadedRegions();
+	void FocusWorld();
 
 	void UpdateTransform() const;
 	void UpdateSelectionBox(bool bSnap);
@@ -112,6 +154,7 @@ protected:
 	bool bIsMeasuring;
 	bool bShowActors;
 	bool bShowGrid;
+	bool bShowMiniMap;
 	bool bFollowPlayerInPIE;
 	FVector2D MouseCursorPos;
 	FVector2D MouseCursorPosWorld;
@@ -167,4 +210,6 @@ protected:
 	mutable double PaintTime;
 
 	TSharedPtr<SWorldPartitionViewportWidget> ViewportWidget;
+	TSharedPtr<STextBlock> TextWorldBoundsInKMWidget;
+	TSharedPtr<STextBlock> TextRulerWidget;
 };
\ No newline at end of file
