OneFilePerActor: Fix deletion of empty __ExternalActor__ directories to not try (and fail) to delete the root __ExternalActor__ directory if it is out of subdirectories.
#rnx

diff --git a/Engine/Source/Editor/WorldPartitionEditor/Private/WorldPartitionEditorModule.cpp b/Engine/Source/Editor/WorldPartitionEditor/Private/WorldPartitionEditorModule.cpp
--- a/Engine/Source/Editor/WorldPartitionEditor/Private/WorldPartitionEditorModule.cpp
+++ b/Engine/Source/Editor/WorldPartitionEditor/Private/WorldPartitionEditorModule.cpp
@@ -862,7 +862,9 @@ void FWorldPartitionEditorModule::FCleanupExternalObjectsEmptyFoldersWorker::DoW
 
 			for (auto& [Folder, NumFiles] : FoldersFilesMap)
 			{
-				if (!NumFiles)
+				if (!NumFiles &&  // Not yet deleted, and now empty
+					Folder != RootPath // Do not delete the RootPath, we should keep it and it would break the assertion that parent is in FolderFilesMap
+					)
 				{
 					if (IFileManager::Get().DeleteDirectory(*Folder, true))
 					{
