World Partition Editor
* Change LoadingRegion/Adapter borders to match Starship style guide. (Selection, Loaded, Hover)
* Remove the Orange Box when PIE, we still have the player arrow in Orange for indication.
* Avoid using colors directly, use Starship style guide.

diff --git a/Engine/Source/Editor/WorldPartitionEditor/Private/WorldPartition/SWorldPartitionEditorGrid2D.cpp b/Engine/Source/Editor/WorldPartitionEditor/Private/WorldPartition/SWorldPartitionEditorGrid2D.cpp
--- a/Engine/Source/Editor/WorldPartitionEditor/Private/WorldPartition/SWorldPartitionEditorGrid2D.cpp
+++ b/Engine/Source/Editor/WorldPartitionEditor/Private/WorldPartition/SWorldPartitionEditorGrid2D.cpp
@@ -21,6 +21,7 @@
 #include "ScopedTransaction.h"
 #include "SWorldPartitionViewportWidget.h"
 #include "SEditorViewportToolBarMenu.h"
+#include "Styling/StyleColors.h"
 #include "TextureResource.h"
 #include "Widgets/Input/SButton.h"
 #include "Widgets/Input/SCheckBox.h"
@@ -1323,8 +1324,8 @@ uint32 SWorldPartitionEditorGrid2D::PaintActors(const FGeometry& AllottedGeometr
 
 			if (bShowBackground)
 			{
-				const FSlateColorBrush BackgroundBrush(FLinearColor::Black);
-				const FLinearColor LabelBackgroundColor = FLinearColor::Black.CopyWithNewOpacity(Color.A * 0.1f);
+				const FSlateColorBrush BackgroundBrush(USlateThemeManager::Get().GetColor(EStyleColor::Black));
+				const FLinearColor LabelBackgroundColor = USlateThemeManager::Get().GetColor(EStyleColor::Black).CopyWithNewOpacity(Color.A * 0.1f);
 
 				FSlateDrawElement::MakeBox(
 					OutDrawElements,
@@ -1350,9 +1351,6 @@ uint32 SWorldPartitionEditorGrid2D::PaintActors(const FGeometry& AllottedGeometr
 
 	if (ShownLoaderInterfaces.Num())
 	{
-		const FLinearColor LoadedActorColor(0.75f, 0.75f, 0.75f, 1.0f);
-		const FLinearColor UnloadedActorColor(0.5f, 0.5f, 0.5f, 1.0f);	
-
 		TArray<FVector2D> LinePoints;
 		LinePoints.SetNum(5);
 
@@ -1394,13 +1392,13 @@ uint32 SWorldPartitionEditorGrid2D::PaintActors(const FGeometry& AllottedGeometr
 					{
 						const FPaintGeometry ActorGeometry = AllottedGeometry.ToPaintGeometry(BottomRight - TopLeft, FSlateLayoutTransform(TopLeft));
 						const float LoaderColorGradient = FMath::Min((ActorViewBox.GetArea() - MinimumAreaCull) / AreaFadeDistance, 1.0f);
-						const FLinearColor LoaderColor = LoaderAdapter->GetColor().IsSet() ? *LoaderAdapter->GetColor() : FColor::White;
+						const FLinearColor LoaderColor = LoaderAdapter->GetColor().IsSet() ? *LoaderAdapter->GetColor() : USlateThemeManager::Get().GetColor(EStyleColor::White);
 						const bool bIsLocalHovered = LocalHoveredLoaderAdapter == LoaderAdapter;
 
 						// Highlight
 						{
-							const FSlateColorBrush LoadedBrush(FLinearColor::White);
-							const FSlateColorBrush UnloadedBrush(FLinearColor::Gray);
+							const FSlateColorBrush LoadedBrush(USlateThemeManager::Get().GetColor(EStyleColor::White));
+							const FSlateColorBrush UnloadedBrush(USlateThemeManager::Get().GetColor(EStyleColor::AccentGray));
 							const FLinearColor LoadedColor(LoaderColor.R, LoaderColor.G, LoaderColor.B, 0.23f * LoaderColorGradient * FullScreenColorGradient);
 							const FLinearColor UnloadedColor(LoaderColor.R * 0.15f, LoaderColor.G * 0.15f, LoaderColor.B * 0.15f, 0.25f * LoaderColorGradient * FullScreenColorGradient);
 
@@ -1444,22 +1442,22 @@ uint32 SWorldPartitionEditorGrid2D::PaintActors(const FGeometry& AllottedGeometr
 							const bool bIsInsideSelection = (SelectBoxGridSnapped.GetVolume() > 0) && SelectBoxGridSnapped.Intersect(AdapterBounds);
 							
 							float OutlineThickness = 2.0f;
-							FLinearColor OutlineColor = FLinearColor::Gray;
+							FLinearColor OutlineColor = USlateThemeManager::Get().GetColor(EStyleColor::AccentGray);
 
 							if (bIsSelected)
 							{
 								OutlineThickness = 4.0f;
-								OutlineColor = FLinearColor::Yellow;
+								OutlineColor = USlateThemeManager::Get().GetColor(EStyleColor::Primary);
 							}
 							else if (bIsInsideSelection)
 							{
 								OutlineThickness = 4.0f;
-								OutlineColor = FLinearColor::Yellow * 0.5f;
+								OutlineColor = USlateThemeManager::Get().GetColor(EStyleColor::Primary) * 0.5f;
 							}
 							else if (bIsLocalHovered)
 							{
 								OutlineThickness = 4.0f;
-								OutlineColor = FLinearColor::White;
+								OutlineColor = USlateThemeManager::Get().GetColor(EStyleColor::White);
 							}
 
 							LinePoints[0] = TopLeft;
@@ -1484,7 +1482,7 @@ uint32 SWorldPartitionEditorGrid2D::PaintActors(const FGeometry& AllottedGeometr
 						// Label
 						{
 							const FString ActorLabel = *LoaderAdapter->GetLabel();
-							const FLinearColor LabelColor = FLinearColor::White.CopyWithNewOpacity(LoaderColorGradient * FullScreenColorGradient);
+							const FLinearColor LabelColor = USlateThemeManager::Get().GetColor(EStyleColor::White).CopyWithNewOpacity(LoaderColorGradient * FullScreenColorGradient);
 							DrawActorLabel(ActorLabel, ActorViewBox.GetCenter(), ActorGeometry, LabelColor, SmallLayoutFont, false);
 						}
 					}
@@ -1495,8 +1493,6 @@ uint32 SWorldPartitionEditorGrid2D::PaintActors(const FGeometry& AllottedGeometr
 
 	if (ActorDescList.Num())
 	{
-		const FLinearColor SelectedActorColor(1.0f, 1.0f, 1.0f, 1.0f);
-
 		TArray<FVector2D> LinePoints;
 		LinePoints.SetNum(5);
 
@@ -1539,15 +1535,15 @@ uint32 SWorldPartitionEditorGrid2D::PaintActors(const FGeometry& AllottedGeometr
 					{
 						if (!ActorLabel.IsNone())
 						{
-							const FLinearColor LabelColor = FLinearColor::Yellow.CopyWithNewOpacity(ActorColorGradient);
+							const FLinearColor LabelColor = USlateThemeManager::Get().GetColor(EStyleColor::AccentOrange).CopyWithNewOpacity(ActorColorGradient);
 							DrawActorLabel(ActorLabel.ToString(), !bIsSelected ? MouseCursorPos : ActorViewBox.GetCenter(), ActorGeometry, LabelColor, SmallLayoutFont, true);
 						}
 
-						ActorColor = FLinearColor::Yellow.CopyWithNewOpacity(ActorColor.A);
+						ActorColor = USlateThemeManager::Get().GetColor(EStyleColor::AccentOrange).CopyWithNewOpacity(ActorColor.A);
 					}
 					else if ((SelectBoxGridSnapped.GetVolume() > 0) && SelectBoxGridSnapped.Intersect(ActorBounds))
 					{
-						ActorColor = FLinearColor::White.CopyWithNewOpacity(ActorColor.A);
+						ActorColor = USlateThemeManager::Get().GetColor(EStyleColor::White).CopyWithNewOpacity(ActorColor.A);
 					}
 
 					FSlateDrawElement::MakeBox(
@@ -1556,7 +1552,7 @@ uint32 SWorldPartitionEditorGrid2D::PaintActors(const FGeometry& AllottedGeometr
 						ActorGeometryShadow,
 						FAppStyle::GetBrush(TEXT("Border")),
 						ESlateDrawEffect::None,
-						FLinearColor::Black.CopyWithNewOpacity(ActorColor.A)
+						USlateThemeManager::Get().GetColor(EStyleColor::Black).CopyWithNewOpacity(ActorColor.A)
 					);
 
 					FSlateDrawElement::MakeBox(
@@ -1608,7 +1604,7 @@ uint32 SWorldPartitionEditorGrid2D::PaintTextInfo(const FGeometry& AllottedGeome
 		AllottedGeometry.ToOffsetPaintGeometry(FVector2D(InX, InY)),
 		LinePoints,
 		ESlateDrawEffect::None,
-		FLinearColor::White);
+		USlateThemeManager::Get().GetColor(EStyleColor::White));
 
 	const float UnitsInRuler = ScaleRulerLength/Scale + 0.05f;// Pixels to world units (+0.05f to accommodate for %.2f)
 	const int32 UnitsInMeter = 100;
@@ -1639,7 +1635,7 @@ uint32 SWorldPartitionEditorGrid2D::PaintTextInfo(const FGeometry& AllottedGeome
 			RulerText,
 			FAppStyle::GetFontStyle("NormalFont"),
 			ESlateDrawEffect::None,
-			FLinearColor::Gray);
+			USlateThemeManager::Get().GetColor(EStyleColor::AccentGray));
 
 		const FVector2D TextSize = FSlateApplication::Get().GetRenderer()->GetFontMeasureService()->Measure(RulerText, FAppStyle::GetFontStyle("NormalFont"));
 
@@ -1651,7 +1647,7 @@ uint32 SWorldPartitionEditorGrid2D::PaintTextInfo(const FGeometry& AllottedGeome
 			RulerText,
 			FAppStyle::GetFontStyle("NormalFont"),
 			ESlateDrawEffect::None,
-			FLinearColor::Gray);
+			USlateThemeManager::Get().GetColor(EStyleColor::AccentGray));
 	}
 		
 	return LayerId + 1;
@@ -1676,7 +1672,7 @@ uint32 SWorldPartitionEditorGrid2D::PaintViewer(const FGeometry& AllottedGeometr
 			FMath::DegreesToRadians(Rotation.Yaw),
 			(Image->ImageSize + ShadowSize) * 0.5f,
 			FSlateDrawElement::RelativeToElement,
-			FLinearColor::Black
+			USlateThemeManager::Get().GetColor(EStyleColor::Black)
 		);
 
 		const FPaintGeometry PaintGeometry = AllottedGeometry.ToPaintGeometry(
@@ -1704,14 +1700,14 @@ uint32 SWorldPartitionEditorGrid2D::PaintViewer(const FGeometry& AllottedGeometr
 	FRotator ObserverRotation;
 	if (GetObserverView(ObserverPosition, ObserverRotation))
 	{
-		MakeRotatedBoxWithShadow(FVector2D(ObserverPosition), ObserverRotation, CameraImage, FLinearColor::White, ShadowSize);
+		MakeRotatedBoxWithShadow(FVector2D(ObserverPosition), ObserverRotation, CameraImage, USlateThemeManager::Get().GetColor(EStyleColor::White), ShadowSize);
 	}
 
 	FVector PlayerPosition;
 	FRotator PlayerRotation;
 	if (GetPlayerView(PlayerPosition, PlayerRotation))
 	{
-		MakeRotatedBoxWithShadow(FVector2D(PlayerPosition), PlayerRotation, CameraImage, FColorList::Orange, ShadowSize);
+		MakeRotatedBoxWithShadow(FVector2D(PlayerPosition), PlayerRotation, CameraImage, USlateThemeManager::Get().GetColor(EStyleColor::AccentOrange), ShadowSize);
 	}
 
 	return LayerId + 1;
@@ -1724,6 +1720,9 @@ uint32 SWorldPartitionEditorGrid2D::PaintSelection(const FGeometry& AllottedGeom
 		TArray<FVector2D> LinePoints;
 		LinePoints.SetNum(5);
 
+		FLinearColor OutlineColor(USlateThemeManager::Get().GetColor(EStyleColor::White));
+		float OutlineThickness = 1.0f;
+
 		// Draw snapped box
 		{
 			FVector2D TopLeftW = FVector2D(SelectBoxGridSnapped.Min);
@@ -1743,8 +1742,8 @@ uint32 SWorldPartitionEditorGrid2D::PaintSelection(const FGeometry& AllottedGeom
 			LinePoints[4] = TopLeft;
 
 			{
-				FSlateColorBrush CellBrush(FLinearColor::White);
-				FLinearColor CellColor(FLinearColor(1, 1, 1, 0.1f));
+				FSlateColorBrush CellBrush(USlateThemeManager::Get().GetColor(EStyleColor::White));
+				FLinearColor CellColor = USlateThemeManager::Get().GetColor(EStyleColor::White).CopyWithNewOpacity(0.1f);
 
 				FPaintGeometry CellGeometry = AllottedGeometry.ToPaintGeometry(
 					BottomRight - TopLeft,
@@ -1767,9 +1766,9 @@ uint32 SWorldPartitionEditorGrid2D::PaintSelection(const FGeometry& AllottedGeom
 				AllottedGeometry.ToPaintGeometry(), 
 				LinePoints, 
 				ESlateDrawEffect::None, 
-				FLinearColor::White, 
+				OutlineColor,
 				true, 
-				2.0f
+				OutlineThickness
 			);
 		}
 
@@ -1797,9 +1796,9 @@ uint32 SWorldPartitionEditorGrid2D::PaintSelection(const FGeometry& AllottedGeom
 				AllottedGeometry.ToPaintGeometry(), 
 				LinePoints, 
 				ESlateDrawEffect::None, 
-				FLinearColor::White, 
+				OutlineColor,
 				true, 
-				2.0f
+				OutlineThickness
 			);
 		}
 	}
@@ -1831,17 +1830,6 @@ int32 SWorldPartitionEditorGrid2D::OnPaint(const FPaintArgs& Args, const FGeomet
 		LayerId = PaintSelection(AllottedGeometry, MyCullingRect, OutDrawElements, ++LayerId);
 		LayerId = PaintSoftwareCursor(AllottedGeometry, MyCullingRect, OutDrawElements, ++LayerId);
 		LayerId = PaintMeasureTool(AllottedGeometry, MyCullingRect, OutDrawElements, ++LayerId);
-		
-		// Draw a surrounding indicator when PIE is active
-		if (UWorldPartition::IsSimulating() || !!GEditor->PlayWorld)
-		{
-			FSlateDrawElement::MakeBox(
-				OutDrawElements,
-				LayerId,
-				AllottedGeometry.ToPaintGeometry(),
-				FAppStyle::GetBrush(TEXT("Graph.PlayInEditor"))
-			);
-		}
 	}
 
 	return SWorldPartitionEditorGrid::OnPaint(Args, AllottedGeometry, MyCullingRect, OutDrawElements, LayerId, InWidgetStyle, bParentEnabled);
@@ -1886,7 +1874,7 @@ int32 SWorldPartitionEditorGrid2D::PaintMinimap(const FGeometry& AllottedGeometr
 			WorldImageGeometry,
 			&WorldMiniMapBrush,
 			ESlateDrawEffect::None,
-			FLinearColor::White
+			USlateThemeManager::Get().GetColor(EStyleColor::White)
 		);
 
 		if (GetWorldPartition()->IsStreamingEnabled())
@@ -1929,8 +1917,8 @@ int32 SWorldPartitionEditorGrid2D::PaintMinimap(const FGeometry& AllottedGeometr
 					FSlateLayoutTransform(ShadowArea.Min)
 				);
 
-				const FSlateColorBrush ShadowdBrush(FLinearColor::Gray);
-				const FLinearColor ShadowColor(0, 0, 0, 0.75f);
+				const FSlateColorBrush ShadowdBrush(USlateThemeManager::Get().GetColor(EStyleColor::AccentGray));
+				const FLinearColor ShadowColor = USlateThemeManager::Get().GetColor(EStyleColor::Black).CopyWithNewOpacity(0.75f);
 
 				FSlateDrawElement::MakeBox(
 					OutDrawElements,
@@ -1961,7 +1949,7 @@ int32 SWorldPartitionEditorGrid2D::PaintMinimap(const FGeometry& AllottedGeometr
 				AllottedGeometry.ToPaintGeometry(),
 				LinePoints,
 				ESlateDrawEffect::None,
-				FLinearColor::Black
+				USlateThemeManager::Get().GetColor(EStyleColor::Black)
 			);
 		}
 
@@ -2014,7 +2002,7 @@ int32 SWorldPartitionEditorGrid2D::PaintMeasureTool(const FGeometry& AllottedGeo
 			AllottedGeometry.ToPaintGeometry(), 
 			LinePoints, 
 			ESlateDrawEffect::None, 
-			FLinearColor::Yellow, 
+			USlateThemeManager::Get().GetColor(EStyleColor::AccentOrange),
 			true, 
 			2.0f
 		);
@@ -2025,7 +2013,7 @@ int32 SWorldPartitionEditorGrid2D::PaintMeasureTool(const FGeometry& AllottedGeo
 			AllottedGeometry,
 			FString::Printf(TEXT("%d"), (int32)FVector2D::Distance(MeasureStart, MeasureEnd)),
 			(MeasureStartScreen + MeasureEndScreen) * 0.5f,
-			FLinearColor::White,
+			USlateThemeManager::Get().GetColor(EStyleColor::White),
 			SmallLayoutFont
 		);
 	}
@@ -2040,8 +2028,8 @@ int32 SWorldPartitionEditorGrid2D::DrawTextLabel(FSlateWindowElementList& OutDra
 	if (LabelTextSize.X > 0)
 	{
 		const FVector2D LabelTextPos = Pos - LabelTextSize * 0.5f;
-		const FLinearColor LabelForegroundColor = FLinearColor::White.CopyWithNewOpacity(Color.A);
-		const FLinearColor LabelShadowColor = FLinearColor::Black.CopyWithNewOpacity(Color.A);
+		const FLinearColor LabelForegroundColor = USlateThemeManager::Get().GetColor(EStyleColor::White).CopyWithNewOpacity(Color.A);
+		const FLinearColor LabelShadowColor = USlateThemeManager::Get().GetColor(EStyleColor::Black).CopyWithNewOpacity(Color.A);
 
 		FSlateDrawElement::MakeText(
 			OutDrawElements,
