Disable CleanupExternalObjectsEmptyFolders as it seems to be taking a very long time for some people.

diff --git a/Engine/Config/BaseEditor.ini b/Engine/Config/BaseEditor.ini
--- a/Engine/Config/BaseEditor.ini
+++ b/Engine/Config/BaseEditor.ini
@@ -31,6 +31,9 @@ ColorFontAtlasSize=512
 ; Content Sub Path being hidden from the Asset view for each mount
 ;+DenyListContentSubPaths=
 
+[WorldPartitionEditor]
+bCleanupExternalObjectsEmptyFolders=false
+
 [/Script/UnrealEd.ThumbnailManager]
 +RenderableThumbnailTypes=(ClassNeedingThumbnailName="/Script/CoreUObject.Class",RendererClassName="/Script/UnrealEd.ClassThumbnailRenderer")
 +RenderableThumbnailTypes=(ClassNeedingThumbnailName="/Script/Engine.Blueprint",RendererClassName="/Script/UnrealEd.BlueprintThumbnailRenderer")
diff --git a/Engine/Source/Editor/WorldPartitionEditor/Private/WorldPartitionEditorModule.cpp b/Engine/Source/Editor/WorldPartitionEditor/Private/WorldPartitionEditorModule.cpp
--- a/Engine/Source/Editor/WorldPartitionEditor/Private/WorldPartitionEditorModule.cpp
+++ b/Engine/Source/Editor/WorldPartitionEditor/Private/WorldPartitionEditorModule.cpp
@@ -26,6 +26,7 @@
 #include "Engine/Level.h"
 
 #include "HAL/FileManager.h"
+#include "Misc/ConfigCacheIni.h"
 #include "Misc/FileHelper.h"
 #include "Misc/MessageDialog.h"
 #include "Misc/StringBuilder.h"
@@ -914,11 +915,16 @@ void FWorldPartitionEditorModule::FCleanupExternalObjectsEmptyFoldersWorker::DoW
 
 void FWorldPartitionEditorModule::CleanupExternalObjectsEmptyFolders()
 {
-	CleanupExternalObjectsEmptyFoldersWorkerAsyncTask = MakeUnique<FAsyncTask<FCleanupExternalObjectsEmptyFoldersWorker>>();
-	CleanupExternalObjectsEmptyFoldersWorkerAsyncTask->StartBackgroundTask();
-	// todo: cleanup content bundles when they are registered?
+	bool bDoCleanup = false;
+	GConfig->GetBool(TEXT("WorldPartitionEditor"), TEXT("bCleanupExternalObjectsEmptyFolders"), bDoCleanup, GEditorIni);
+	if (bDoCleanup)
+	{
+		CleanupExternalObjectsEmptyFoldersWorkerAsyncTask = MakeUnique<FAsyncTask<FCleanupExternalObjectsEmptyFoldersWorker>>();
+		CleanupExternalObjectsEmptyFoldersWorkerAsyncTask->StartBackgroundTask();
+		// todo: cleanup content bundles when they are registered?
 
-	FCoreDelegates::OnPostEngineInit.AddLambda([this]() { WaitForCleanupExternalObjectsEmptyFolders(); });
+		FCoreDelegates::OnPostEngineInit.AddLambda([this]() { WaitForCleanupExternalObjectsEmptyFolders(); });
+	}
 }
 
 void FWorldPartitionEditorModule::WaitForCleanupExternalObjectsEmptyFolders()
