World Partition Editor: fix editor not focusing on world bounds since recent changes.

diff --git a/Engine/Source/Editor/WorldPartitionEditor/Private/WorldPartition/SWorldPartitionEditorGrid2D.cpp b/Engine/Source/Editor/WorldPartitionEditor/Private/WorldPartition/SWorldPartitionEditorGrid2D.cpp
--- a/Engine/Source/Editor/WorldPartitionEditor/Private/WorldPartition/SWorldPartitionEditorGrid2D.cpp
+++ b/Engine/Source/Editor/WorldPartitionEditor/Private/WorldPartition/SWorldPartitionEditorGrid2D.cpp
@@ -1818,7 +1818,7 @@ int32 SWorldPartitionEditorGrid2D::OnPaint(const FPaintArgs& Args, const FGeomet
 
 		if (bResetView)
 		{
-			const_cast<SWorldPartitionEditorGrid2D*>(this)->FocusSelection();
+			const_cast<SWorldPartitionEditorGrid2D*>(this)->FocusWorld();
 		}
 
 		UpdateTransform();
