World Partition: move data layer and hlod layer asset into the world category.

diff --git a/Engine/Source/Developer/AssetTools/Private/AssetTools.cpp b/Engine/Source/Developer/AssetTools/Private/AssetTools.cpp
--- a/Engine/Source/Developer/AssetTools/Private/AssetTools.cpp
+++ b/Engine/Source/Developer/AssetTools/Private/AssetTools.cpp
@@ -1348,6 +1348,7 @@ UAssetToolsImpl::UAssetToolsImpl(const FObjectInitializer& ObjectInitializer)
 	AllocatedCategoryBits.Add(TEXT("Gameplay"), FAdvancedAssetCategory(EAssetTypeCategories::Gameplay, LOCTEXT("GameplayAssetCategory", "Gameplay")));
 	AllocatedCategoryBits.Add(TEXT("Media"), FAdvancedAssetCategory(EAssetTypeCategories::Media, LOCTEXT("MediaAssetCategory", "Media")));
 	AllocatedCategoryBits.Add(TEXT("Texture"), FAdvancedAssetCategory(EAssetTypeCategories::Textures, LOCTEXT("TextureAssetCategory", "Texture")));
+	AllocatedCategoryBits.Add(TEXT("World"), FAdvancedAssetCategory(EAssetTypeCategories::World, LOCTEXT("WorldAssetCategory", "World")));
 
 	EAssetTypeCategories::Type InputCategoryBit = RegisterAdvancedAssetCategory(FName(TEXT("Input")), LOCTEXT("InputAssetsCategory", "Input"));
 	
diff --git a/Engine/Source/Developer/AssetTools/Public/AssetTypeCategories.h b/Engine/Source/Developer/AssetTools/Public/AssetTypeCategories.h
--- a/Engine/Source/Developer/AssetTools/Public/AssetTypeCategories.h
+++ b/Engine/Source/Developer/AssetTools/Public/AssetTypeCategories.h
@@ -26,9 +26,10 @@ namespace EAssetTypeCategories
 		Blueprint = 1 << 8,
 		Media = 1 << 9,
 		Textures = 1 << 10,
+		World = 1 << 11,
 
 		// Items below this will be allocated at runtime via RegisterAdvancedAssetCategory
-		FirstUser = 1 << 11,
+		FirstUser = 1 << 12,
 		LastUser = 1 << 31,
 		// Last allowed value is 1 << 31
 	};
diff --git a/Engine/Source/Editor/WorldPartitionEditor/Private/WorldPartition/DataLayers/DataLayerAssetTypeActions.cpp b/Engine/Source/Editor/WorldPartitionEditor/Private/WorldPartition/DataLayers/DataLayerAssetTypeActions.cpp
new file mode 100644
--- /dev/null
+++ b/Engine/Source/Editor/WorldPartitionEditor/Private/WorldPartition/DataLayers/DataLayerAssetTypeActions.cpp
@@ -0,0 +1,29 @@
+// Copyright Epic Games, Inc. All Rights Reserved.
+
+#include "DataLayerAssetTypeActions.h"
+#include "WorldPartition/DataLayer/DataLayerAsset.h"
+
+FText FDataLayerAssetTypeActions::GetName() const
+{
+	return NSLOCTEXT("AssetTypeActions", "AssetTypeActions_DataLayerAsset", "DataLayer Asset");
+}
+
+FColor FDataLayerAssetTypeActions::GetTypeColor() const
+{
+	return FColor(0, 200, 200);
+}
+
+UClass* FDataLayerAssetTypeActions::GetSupportedClass() const
+{
+	return UDataLayerAsset::StaticClass();
+}
+
+uint32 FDataLayerAssetTypeActions::GetCategories()
+{
+	return EAssetTypeCategories::World;
+}
+
+bool FDataLayerAssetTypeActions::CanLocalize() const
+{
+	return false;
+}
diff --git a/Engine/Source/Editor/WorldPartitionEditor/Private/WorldPartition/DataLayers/DataLayerAssetTypeActions.h b/Engine/Source/Editor/WorldPartitionEditor/Private/WorldPartition/DataLayers/DataLayerAssetTypeActions.h
new file mode 100644
--- /dev/null
+++ b/Engine/Source/Editor/WorldPartitionEditor/Private/WorldPartition/DataLayers/DataLayerAssetTypeActions.h
@@ -0,0 +1,18 @@
+// Copyright Epic Games, Inc. All Rights Reserved.
+
+#pragma once
+
+#include "CoreMinimal.h"
+#include "AssetTypeActions_Base.h"
+
+class FDataLayerAssetTypeActions : public FAssetTypeActions_Base
+{
+public:
+	// IAssetTypeActions interface
+	virtual FText GetName() const override;
+	virtual FColor GetTypeColor() const override;
+	virtual UClass* GetSupportedClass() const override;
+	virtual uint32 GetCategories() override;
+	virtual bool CanLocalize() const override;
+	// End of IAssetTypeActions interface
+};
diff --git a/Engine/Source/Editor/WorldPartitionEditor/Private/WorldPartition/HLOD/HLODLayerAssetTypeActions.cpp b/Engine/Source/Editor/WorldPartitionEditor/Private/WorldPartition/HLOD/HLODLayerAssetTypeActions.cpp
--- a/Engine/Source/Editor/WorldPartitionEditor/Private/WorldPartition/HLOD/HLODLayerAssetTypeActions.cpp
+++ b/Engine/Source/Editor/WorldPartitionEditor/Private/WorldPartition/HLOD/HLODLayerAssetTypeActions.cpp
@@ -20,7 +20,7 @@ UClass* FHLODLayerAssetTypeActions::GetSupportedClass() const
 
 uint32 FHLODLayerAssetTypeActions::GetCategories()
 {
-	return EAssetTypeCategories::Misc;
+	return EAssetTypeCategories::World;
 }
 
 bool FHLODLayerAssetTypeActions::CanLocalize() const
diff --git a/Engine/Source/Editor/WorldPartitionEditor/Private/WorldPartitionEditorModule.cpp b/Engine/Source/Editor/WorldPartitionEditor/Private/WorldPartitionEditorModule.cpp
--- a/Engine/Source/Editor/WorldPartitionEditor/Private/WorldPartitionEditorModule.cpp
+++ b/Engine/Source/Editor/WorldPartitionEditor/Private/WorldPartitionEditorModule.cpp
@@ -8,6 +8,7 @@
 #include "WorldPartition/WorldPartitionLandscapeSplineMeshesBuilder.h"
 #include "WorldPartition/WorldPartitionActorLoaderInterface.h"
 #include "WorldPartition/HLOD/HLODLayerAssetTypeActions.h"
+#include "WorldPartition/DataLayers/DataLayerAssetTypeActions.h"
 #include "WorldPartition/SWorldPartitionEditor.h"
 #include "WorldPartition/SWorldPartitionEditorGridSpatialHash.h"
 #include "WorldPartition/Customizations/WorldPartitionDetailsCustomization.h"
@@ -174,6 +175,8 @@ void FWorldPartitionEditorModule::StartupModule()
 	IAssetTools& AssetTools = FModuleManager::LoadModuleChecked<FAssetToolsModule>("AssetTools").Get();
 	HLODLayerAssetTypeActions = MakeShareable(new FHLODLayerAssetTypeActions);
 	AssetTools.RegisterAssetTypeActions(HLODLayerAssetTypeActions.ToSharedRef());
+	DataLayerAssetTypeActions = MakeShareable(new FDataLayerAssetTypeActions);
+	AssetTools.RegisterAssetTypeActions(DataLayerAssetTypeActions.ToSharedRef());
 
 	FPropertyEditorModule& PropertyEditor = FModuleManager::LoadModuleChecked<FPropertyEditorModule>("PropertyEditor");
 	PropertyEditor.RegisterCustomClassLayout("WorldPartition", FOnGetDetailCustomizationInstance::CreateStatic(&FWorldPartitionDetails::MakeInstance));
@@ -210,17 +213,26 @@ void FWorldPartitionEditorModule::ShutdownModule()
 		UToolMenus::UnregisterOwner(this);
 	}
 
-	// Unregister the HLODLayer asset type actions
-	if (HLODLayerAssetTypeActions.IsValid())
+	if (FModuleManager::Get().IsModuleLoaded("AssetTools"))
 	{
-		if (FModuleManager::Get().IsModuleLoaded("AssetTools"))
+		IAssetTools& AssetTools = FModuleManager::GetModuleChecked<FAssetToolsModule>("AssetTools").Get();
+
+		// Unregister the HLODLayer asset type actions
+		if (HLODLayerAssetTypeActions.IsValid())
 		{
-			IAssetTools& AssetTools = FModuleManager::GetModuleChecked<FAssetToolsModule>("AssetTools").Get();
 			AssetTools.UnregisterAssetTypeActions(HLODLayerAssetTypeActions.ToSharedRef());
 		}
-		HLODLayerAssetTypeActions.Reset();
+
+		// Unregister the DataLayer asset type actions
+		if (DataLayerAssetTypeActions.IsValid())
+		{
+			AssetTools.UnregisterAssetTypeActions(DataLayerAssetTypeActions.ToSharedRef());
+		}
 	}
 
+	HLODLayerAssetTypeActions.Reset();
+	DataLayerAssetTypeActions.Reset();
+
 	if (FModuleManager::Get().IsModuleLoaded("PropertyEditor"))
 	{
 		FPropertyEditorModule& PropertyEditor = FModuleManager::GetModuleChecked<FPropertyEditorModule>("PropertyEditor");
