Missing file that wasn't checked out from previous submit.

diff --git a/Engine/Source/Editor/WorldPartitionEditor/Public/WorldPartitionEditorModule.h b/Engine/Source/Editor/WorldPartitionEditor/Public/WorldPartitionEditorModule.h
--- a/Engine/Source/Editor/WorldPartitionEditor/Public/WorldPartitionEditorModule.h
+++ b/Engine/Source/Editor/WorldPartitionEditor/Public/WorldPartitionEditorModule.h
@@ -150,6 +150,7 @@ private:
 	FDelegateHandle LevelEditorExtenderDelegateHandle;
 
 	TSharedPtr<class FHLODLayerAssetTypeActions> HLODLayerAssetTypeActions;
+	TSharedPtr<class FDataLayerAssetTypeActions> DataLayerAssetTypeActions;
 
 	TWeakPtr<SDockTab> WorldPartitionTab;
 
