Fixed typo/warning from last CL
#rnx

diff --git a/Engine/Source/Editor/WorldPartitionEditor/Private/WorldPartition/SWorldPartitionEditorGrid2D.cpp b/Engine/Source/Editor/WorldPartitionEditor/Private/WorldPartition/SWorldPartitionEditorGrid2D.cpp
--- a/Engine/Source/Editor/WorldPartitionEditor/Private/WorldPartition/SWorldPartitionEditorGrid2D.cpp
+++ b/Engine/Source/Editor/WorldPartitionEditor/Private/WorldPartition/SWorldPartitionEditorGrid2D.cpp
@@ -1995,7 +1995,7 @@ int32 SWorldPartitionEditorGrid2D::PaintMinimap(const FGeometry& AllottedGeometr
 				IAllocatedVirtualTexture* AllocatedVT = VTResource->AcquireAllocatedVT();
 
 				IRendererModule& RenderModule = GetRendererModule();
-				RenderModule.RequestVirtualTextureTilesForRegion(AllocatedVT, ScreenSpaceSize, ViewportPositon, ViewportSize, UV0, UV1, MipLevel);
+				RenderModule.RequestVirtualTextureTiles(AllocatedVT, ScreenSpaceSize, ViewportPositon, ViewportSize, UV0, UV1, MipLevel);
 				RenderModule.LoadPendingVirtualTextureTiles(RHICmdList, InFeatureLevel);
 			});
 		}
