World Partition - Minimap: Fix viewport offset used to fetch the minimap VT tiles

diff --git a/Engine/Source/Editor/WorldPartitionEditor/Private/WorldPartition/SWorldPartitionEditorGrid2D.cpp b/Engine/Source/Editor/WorldPartitionEditor/Private/WorldPartition/SWorldPartitionEditorGrid2D.cpp
--- a/Engine/Source/Editor/WorldPartitionEditor/Private/WorldPartition/SWorldPartitionEditorGrid2D.cpp
+++ b/Engine/Source/Editor/WorldPartitionEditor/Private/WorldPartition/SWorldPartitionEditorGrid2D.cpp
@@ -1977,7 +1977,7 @@ int32 SWorldPartitionEditorGrid2D::PaintMinimap(const FGeometry& AllottedGeometr
 			FVirtualTexture2DResource* VTResource = static_cast<FVirtualTexture2DResource*>(Texture2D->GetResource());
 			const FVector2D ViewportSize = AllottedGeometry.GetLocalSize();
 			const FVector2D ScreenSpaceSize = WorldImageGeometry.GetLocalSize();
-			const FVector2D ViewportPositon = WorldImageGeometry.GetAccumulatedRenderTransform().GetTranslation() + AllottedGeometry.GetAbsolutePosition();
+			const FVector2D ViewportPositon = AllottedGeometry.GetAbsolutePosition() - WorldImageGeometry.GetAccumulatedRenderTransform().GetTranslation();
 
 			FBox2D UVRegion = WorldMiniMapBrush.GetUVRegion();
 			const FVector2D UV0 = UVRegion.Min;
