World Partition Editor
* Add option (slider) to play with the unloaded region opacity (default is now ~10% brighter).
* Tweak clamp between good values with the rest of the features such as highligting regions, loaded regions.
* Value is saved in PerProjectUserSettings.

diff --git a/Engine/Source/Editor/WorldPartitionEditor/Private/WorldPartition/SWorldPartitionEditorGrid2D.cpp b/Engine/Source/Editor/WorldPartitionEditor/Private/WorldPartition/SWorldPartitionEditorGrid2D.cpp
--- a/Engine/Source/Editor/WorldPartitionEditor/Private/WorldPartition/SWorldPartitionEditorGrid2D.cpp
+++ b/Engine/Source/Editor/WorldPartitionEditor/Private/WorldPartition/SWorldPartitionEditorGrid2D.cpp
@@ -25,6 +25,7 @@
 #include "TextureResource.h"
 #include "Widgets/Input/SButton.h"
 #include "Widgets/Input/SCheckBox.h"
+#include "Widgets/Input/SSlider.h"
 #include "Widgets/Layout/SUniformGridPanel.h"
 #include "Framework/MultiBox/MultiBoxBuilder.h"
 #include "WorldBrowserModule.h"
@@ -283,7 +284,9 @@ void SWorldPartitionEditorGrid2D::FEditorCommands::RegisterCommands()
 
 void SWorldPartitionEditorGrid2D::SToolBar::Construct(const FArguments& InArgs)
 {
-	CommandList = InArgs._CommandList;
+	WPEditorGrid2D = InArgs._WPEditorGrid2D;
+	check(WPEditorGrid2D);
+	CommandList = WPEditorGrid2D->CommandList;
 
 	auto MakeQuickActionsToolBarWidget = [this]()
 	{
@@ -383,17 +386,38 @@ EVisibility SWorldPartitionEditorGrid2D::SToolBar::IsOptionsMenuVisible() const
 	return EVisibility::Visible;
 }
 
+TSharedRef<SWidget> SWorldPartitionEditorGrid2D::SToolBar::GenerateUnloadedOpacitySlider() const
+{	
+	const float OpacityMinValue = 0.5f;
+	const float OpacityMaxValue = 0.8f;
+
+	return
+		SNew(SSlider)
+		.ToolTipText(LOCTEXT("MinimapUnloadedOpactiyToolTip", "Adjust the opacity of the unloaded regions of the minimap."))
+		.MinValue(OpacityMinValue)
+		.MaxValue(OpacityMaxValue)
+		.Value(WPEditorGrid2D, &SWorldPartitionEditorGrid2D::GetMiniMapUnloadedOpacity)
+		.OnValueChanged(WPEditorGrid2D, &SWorldPartitionEditorGrid2D::SetMiniMapUnloadedOpacity)
+		.OnMouseCaptureEnd(WPEditorGrid2D, &SWorldPartitionEditorGrid2D::SaveMiniMapUnloadedOpacityUserSetting)
+		.IsEnabled(WPEditorGrid2D, &SWorldPartitionEditorGrid2D::IsMiniMapUnloadedOpacityEnabled);
+}
+
 TSharedRef<SWidget> SWorldPartitionEditorGrid2D::SToolBar::GenerateOptionsMenu() const
 {
 	static const FName MenuName(TEXT("WorldPartition.OptionsMenu"));
-	static const FName SectionName(TEXT("Options"));
 
 	const FEditorCommands& Commands = FEditorCommands::Get();
-	UToolMenu* ShowMenu = UToolMenus::Get()->RegisterMenu(MenuName);
+	UToolMenu* OptionsMenu = UToolMenus::Get()->RegisterMenu(MenuName);
 
-	FToolMenuSection& Section = ShowMenu->FindOrAddSection(SectionName);
-	Section.AddMenuEntry(Commands.FollowPlayerInPIE);
-	Section.AddMenuEntry(Commands.BugItGoLoadRegion);
+	static const FName SectionGeneralName(TEXT("Options.General"));
+	FToolMenuSection& SectionGeneral = OptionsMenu->FindOrAddSection(SectionGeneralName);
+	SectionGeneral.AddMenuEntry(Commands.FollowPlayerInPIE);
+	SectionGeneral.AddMenuEntry(Commands.BugItGoLoadRegion);
+
+	static const FName SectionMinimapName(TEXT("Options.Minimap"));
+	FToolMenuSection& SectionMinimap = OptionsMenu->AddSection(SectionMinimapName, LOCTEXT("WorldPartitionOptionsMenuMinimap", "Minimap"));
+	static const FName SliderUnloadedOpacityName(TEXT("UnloadedOpacity"));
+	SectionMinimap.AddEntry(FToolMenuEntry::InitWidget(SliderUnloadedOpacityName, GenerateUnloadedOpacitySlider(), LOCTEXT("UnloadedOpacity", "Unloaded Opacity"), true));	
 
 	return UToolMenus::Get()->GenerateWidget(MenuName, FToolMenuContext(CommandList));
 }
@@ -455,6 +479,8 @@ SWorldPartitionEditorGrid2D::SWorldPartitionEditorGrid2D()
 	
 	FWorldBrowserModule& WorldBrowserModule = FModuleManager::LoadModuleChecked<FWorldBrowserModule>("WorldBrowser");
 	WorldBrowserModule.OnShutdown().AddLambda([](){ FEditorCommands::Unregister(); });
+
+	MiniMapUnloadedOpacity = GetMutableDefault<UWorldPartitionEditorPerProjectUserSettings>()->GetMinimapUnloadedOpacity();
 }
 
 SWorldPartitionEditorGrid2D::~SWorldPartitionEditorGrid2D()
@@ -490,7 +516,7 @@ void SWorldPartitionEditorGrid2D::Construct(const FArguments& InArgs)
 		.Padding(ToolbarSlotPadding)
 		[
 			SNew(SToolBar)
-			.CommandList(CommandList)
+			.WPEditorGrid2D(this)
 		]
 
 		// Bottom section
@@ -619,6 +645,37 @@ void SWorldPartitionEditorGrid2D::BindCommands()
 	CommandList->MapAction(Commands.FocusWorld, FExecuteAction::CreateSP(this, &SWorldPartitionEditorGrid2D::FocusWorld), FCanExecuteAction::CreateLambda([this]() { return IsInteractive(); }));
 }
 
+void SWorldPartitionEditorGrid2D::SaveMiniMapUnloadedOpacityUserSetting()
+{
+	GetMutableDefault<UWorldPartitionEditorPerProjectUserSettings>()->SetMinimapUnloadedOpacity(MiniMapUnloadedOpacity);
+}
+
+bool SWorldPartitionEditorGrid2D::IsMiniMapUnloadedOpacityEnabled() const
+{
+	if (!bShowMiniMap)
+	{
+		return false;
+	}
+
+	if (!GetDefault<UWorldPartitionEditorSettings>()->bEnableLoadingInEditor)
+	{
+		return false;
+	}
+
+	if (!GetWorldPartition()->IsStreamingEnabled())
+	{
+		return false;
+	}
+
+	UTexture2D* Texture2D = Cast<UTexture2D>(WorldMiniMapBrush.GetResourceObject());
+	if (Texture2D == nullptr)
+	{
+		return false;
+	}
+
+	return true;
+}
+
 void SWorldPartitionEditorGrid2D::UpdateWorldMiniMapDetails()
 {
 	if (AWorldPartitionMiniMap* WorldMiniMap = FWorldPartitionMiniMapHelper::GetWorldPartitionMiniMap(GetWorld()))
@@ -1937,7 +1994,7 @@ int32 SWorldPartitionEditorGrid2D::PaintMinimap(const FGeometry& AllottedGeometr
 				);
 
 				const FSlateColorBrush ShadowdBrush(USlateThemeManager::Get().GetColor(EStyleColor::AccentGray));
-				const FLinearColor ShadowColor = USlateThemeManager::Get().GetColor(EStyleColor::Black).CopyWithNewOpacity(0.75f);
+				const FLinearColor ShadowColor = USlateThemeManager::Get().GetColor(EStyleColor::Black).CopyWithNewOpacity(MiniMapUnloadedOpacity);
 
 				FSlateDrawElement::MakeBox(
 					OutDrawElements,
diff --git a/Engine/Source/Editor/WorldPartitionEditor/Public/WorldPartition/SWorldPartitionEditorGrid2D.h b/Engine/Source/Editor/WorldPartitionEditor/Public/WorldPartition/SWorldPartitionEditorGrid2D.h
--- a/Engine/Source/Editor/WorldPartitionEditor/Public/WorldPartition/SWorldPartitionEditorGrid2D.h
+++ b/Engine/Source/Editor/WorldPartitionEditor/Public/WorldPartition/SWorldPartitionEditorGrid2D.h
@@ -59,21 +59,24 @@ protected:
 	// In-viewport toolbar widget used in the world partition editor
 	class SToolBar : public SViewportToolBar
 	{
-		public:
-			SLATE_BEGIN_ARGS(SToolBar) {}
-				SLATE_ARGUMENT(TSharedPtr<FUICommandList>, CommandList)
-			SLATE_END_ARGS()
+	public:
+		SLATE_BEGIN_ARGS(SToolBar) {}
+			SLATE_ARGUMENT(SWorldPartitionEditorGrid2D*, WPEditorGrid2D)
+		SLATE_END_ARGS()
 
-			void Construct(const FArguments& InArgs);
+		void Construct(const FArguments& InArgs);
 
 	private:
 		EVisibility IsOptionsMenuVisible() const;
 		
+		TSharedRef<SWidget> GenerateUnloadedOpacitySlider() const;
+
 		TSharedRef<SWidget> GenerateOptionsMenu() const;
 		TSharedRef<SWidget> GenerateShowMenu() const;
 		TSharedRef<SWidget> GenerateBuildMenu() const;
 		
 		TSharedPtr<FUICommandList> CommandList;
+		SWorldPartitionEditorGrid2D* WPEditorGrid2D;
 	};
 
 
@@ -207,6 +210,11 @@ protected:
 
 	// Minimap
 	void UpdateWorldMiniMapDetails();
+	void SaveMiniMapUnloadedOpacityUserSetting();
+	bool IsMiniMapUnloadedOpacityEnabled() const;
+	float GetMiniMapUnloadedOpacity() const { return MiniMapUnloadedOpacity; };
+	void SetMiniMapUnloadedOpacity(float InOpacity) { MiniMapUnloadedOpacity = InOpacity; };
+	float MiniMapUnloadedOpacity;
 	FBox2D WorldMiniMapBounds;
 	FSlateBrush WorldMiniMapBrush;
 
diff --git a/Engine/Source/Runtime/Engine/Public/WorldPartition/WorldPartitionEditorPerProjectUserSettings.h b/Engine/Source/Runtime/Engine/Public/WorldPartition/WorldPartitionEditorPerProjectUserSettings.h
--- a/Engine/Source/Runtime/Engine/Public/WorldPartition/WorldPartitionEditorPerProjectUserSettings.h
+++ b/Engine/Source/Runtime/Engine/Public/WorldPartition/WorldPartitionEditorPerProjectUserSettings.h
@@ -60,6 +60,7 @@ public:
 		, bHighlightSelectedDataLayers(true)
 		, bHideLevelInstanceContent(true)
 		, bDisableLoadingOfLastLoadedRegions(false)
+		, MinimapUnloadedOpacity(0.66f)
 #endif
 	{}
 
@@ -103,6 +104,20 @@ public:
 		}
 	}
 
+	float GetMinimapUnloadedOpacity() const
+	{
+		return MinimapUnloadedOpacity;
+	}
+
+	void SetMinimapUnloadedOpacity(float InMinimapUnloadedOpacity)
+	{
+		if (MinimapUnloadedOpacity != InMinimapUnloadedOpacity)
+		{
+			MinimapUnloadedOpacity = InMinimapUnloadedOpacity;
+			SaveConfig();
+		}
+	}
+
 	ENGINE_API TArray<FName> GetWorldDataLayersNotLoadedInEditor(UWorld* InWorld) const;
 	ENGINE_API TArray<FName> GetWorldDataLayersLoadedInEditor(UWorld* InWorld) const;
 	
@@ -156,6 +171,9 @@ private:
 	UPROPERTY(config)
 	uint32 bShowCellCoords : 1;
 
+	UPROPERTY(config)
+	float MinimapUnloadedOpacity;
+
 	UPROPERTY(config)
 	TMap<TSoftObjectPtr<UWorld>, FWorldPartitionPerWorldSettings> PerWorldEditorSettings;
 #endif
