World Partition Editor: fix selecting loaded actors not showing in the world partition editor window, broken in CL 21336657.

diff --git a/Engine/Source/Editor/WorldPartitionEditor/Private/WorldPartition/SWorldPartitionEditorGrid2D.cpp b/Engine/Source/Editor/WorldPartitionEditor/Private/WorldPartition/SWorldPartitionEditorGrid2D.cpp
--- a/Engine/Source/Editor/WorldPartitionEditor/Private/WorldPartition/SWorldPartitionEditorGrid2D.cpp
+++ b/Engine/Source/Editor/WorldPartitionEditor/Private/WorldPartition/SWorldPartitionEditorGrid2D.cpp
@@ -186,10 +186,7 @@ void ForEachIntersectingLoaderAdapters(UWorldPartition* WorldPartition, const FB
 		{
 			if (LoaderAdapter->GetBoundingBox().IsSet() && IsBoundsSelected(SelectBox, *LoaderAdapter->GetBoundingBox()))
 			{
-				if (!Func(EditorLoaderAdapter))
-				{
-					return;
-				}
+				Func(EditorLoaderAdapter);
 			}
 		}
 	}
@@ -204,10 +201,7 @@ void ForEachIntersectingLoaderAdapters(UWorldPartition* WorldPartition, const FB
 				{
 					if (IsBoundsSelected(SelectBox, ActorDesc->GetEditorBounds()))
 					{
-						if (!Func(Actor))
-						{
-							return false;
-						}
+						Func(Actor);
 					}
 				}
 			}
@@ -1009,7 +1003,6 @@ FReply SWorldPartitionEditorGrid2D::OnMouseButtonUp(const FGeometry& MyGeometry,
 				ForEachIntersectingLoaderAdapters(GetWorldPartition(), SelectBoxGridSnapped, [&LoaderAdaptersToSelect](UObject* AdapterObject)
 				{
 					LoaderAdaptersToSelect.Add(AdapterObject);
-					return true;
 				});
 			}
 			else if (HoveredLoaderInterface.IsValid())
@@ -1230,6 +1223,7 @@ void SWorldPartitionEditorGrid2D::Tick(const FGeometry& AllottedGeometry, const
 
 	ShownActorGuids.Reset();
 	DirtyActorGuids.Reset();
+	SelectedActorGuids.Reset();
 	ShownLoaderInterfaces.Reset(); 
 	
 	for (UWorldPartitionEditorLoaderAdapter* EditorLoaderAdapter : GetRegisteredEditorLoaderAdapters(GetWorldPartition()))
@@ -1353,6 +1347,7 @@ void SWorldPartitionEditorGrid2D::Tick(const FGeometry& AllottedGeometry, const
 		if (AActor* Actor = Cast<AActor>(*It))
 		{
 			ShownActorGuids.Add(Actor->GetActorGuid());
+			SelectedActorGuids.Add(Actor->GetActorGuid());
 		}
 	}
 
@@ -1643,8 +1638,7 @@ uint32 SWorldPartitionEditorGrid2D::PaintActors(const FGeometry& AllottedGeometr
 			};
 
 			const FBox ActorBounds = ActorDescView.GetEditorBounds();
-			const AActor* Actor = ActorDescView.GetActor();
-			const bool bIsSelected = Actor ? Actor->IsSelected() : WorldPartitionSubsystem->SelectedActorDescs.Contains(ActorDescView.GetActorDesc());
+			const bool bIsSelected = SelectedActorGuids.Contains(ActorDescView.GetGuid()) || WorldPartitionSubsystem->SelectedActorDescs.Contains(ActorDescView.GetActorDesc());
 			const bool bIsSpatiallyLoaded = ActorDescView.GetIsSpatiallyLoaded();
 			const FName ActorLabel = ActorDescView.GetActorLabel();
 
diff --git a/Engine/Source/Editor/WorldPartitionEditor/Public/WorldPartition/SWorldPartitionEditorGrid2D.h b/Engine/Source/Editor/WorldPartitionEditor/Public/WorldPartition/SWorldPartitionEditorGrid2D.h
--- a/Engine/Source/Editor/WorldPartitionEditor/Public/WorldPartition/SWorldPartitionEditorGrid2D.h
+++ b/Engine/Source/Editor/WorldPartitionEditor/Public/WorldPartition/SWorldPartitionEditorGrid2D.h
@@ -203,6 +203,7 @@ protected:
 	// Updated every tick
 	TSet<FGuid> ShownActorGuids;
 	TSet<FGuid> DirtyActorGuids;
+	TSet<FGuid> SelectedActorGuids;
 	FLoaderInterfaceSet ShownLoaderInterfaces;
 	FLoaderInterfaceSet HoveredLoaderInterfaces;
 	FLoaderInterfaceStack HoveredLoaderInterfacesStack;
