World Partition Editor: fix computation overflow in int64->float->int64 conversions.

diff --git a/Engine/Source/Editor/WorldPartitionEditor/Private/WorldPartition/SWorldPartitionEditorGridSpatialHash.cpp b/Engine/Source/Editor/WorldPartitionEditor/Private/WorldPartition/SWorldPartitionEditorGridSpatialHash.cpp
--- a/Engine/Source/Editor/WorldPartitionEditor/Private/WorldPartition/SWorldPartitionEditorGridSpatialHash.cpp
+++ b/Engine/Source/Editor/WorldPartitionEditor/Private/WorldPartition/SWorldPartitionEditorGridSpatialHash.cpp
@@ -114,7 +114,7 @@ int32 SWorldPartitionEditorGridSpatialHash::PaintGrid(const FGeometry& AllottedG
 	// Paint grid lines
 	if (ToBox2D(VisibleGridRectWorld).GetArea() > 0.0f)
 	{
-		auto PaintGridLines = [this, &OutDrawElements, &LayerId, &AllottedGeometry](const FBox& VisibleGridRectWorld, float EffectiveCellSize, float Thickness, const FLinearColor& Color)
+		auto PaintGridLines = [this, &OutDrawElements, &LayerId, &AllottedGeometry](const FBox& VisibleGridRectWorld, int64 EffectiveCellSize, float Thickness, const FLinearColor& Color)
 		{
 			if (Color.A > 0.0f)
 			{
@@ -252,4 +252,4 @@ int32 SWorldPartitionEditorGridSpatialHash::PaintGrid(const FGeometry& AllottedG
 	return SWorldPartitionEditorGrid2D::PaintGrid(AllottedGeometry, MyCullingRect, OutDrawElements, LayerId);
 }
 
-#undef LOCTEXT_NAMESPACE
\ No newline at end of file
+#undef LOCTEXT_NAMESPACE
