World Partition - HLOD: Added support for showing HLODs in editor
* For now, disabled by default, can be enabled with "wp.Editor.HLOD 1"
* Currently, the HLODs that are loaded are those that are non spatially loaded & those that aren't "instanced HLOD".
* Will show transition between HLODs depending on the distance
* Loading regions or pinning actors will hide the HLODs intersecting those bounds

diff --git a/Engine/Source/Editor/WorldPartitionEditor/Private/WorldPartition/HLOD/HLODEditorData.cpp b/Engine/Source/Editor/WorldPartitionEditor/Private/WorldPartition/HLOD/HLODEditorData.cpp
new file mode 100644
--- /dev/null
+++ b/Engine/Source/Editor/WorldPartitionEditor/Private/WorldPartition/HLOD/HLODEditorData.cpp
@@ -0,0 +1,227 @@
+// Copyright Epic Games, Inc. All Rights Reserved.
+
+#include "WorldPartition/HLOD/HLODEditorData.h"
+
+#include "WorldPartition/HLOD/HLODActor.h"
+#include "WorldPartition/HLOD/HLODActorDesc.h"
+#include "WorldPartition/HLOD/HLODLoaderAdapter.h"
+
+#include "WorldPartition/ActorDescContainerCollection.h"
+#include "WorldPartition/WorldPartition.h"
+#include "WorldPartition/WorldPartitionHandle.h"
+#include "WorldPartition/WorldPartitionHelpers.h"
+
+
+FWorldPartitionHLODEditorData::FWorldPartitionHLODEditorData(UWorldPartition* InWorldPartition)
+	: WorldPartition(InWorldPartition)
+	, LastStateUpdate(INDEX_NONE)
+{
+	// First pass, build mapping of GUID to HLOD scene node
+	for (FActorDescContainerCollection::TIterator<AWorldPartitionHLOD> HLODIterator(WorldPartition); HLODIterator; ++HLODIterator)
+	{
+		const FHLODActorDesc& HLODActorDesc = **HLODIterator;
+		const FGuid& HLODActorGuid = HLODActorDesc.GetGuid();
+
+		TUniquePtr<FHLODSceneNode>& HLODSceneNode = HLODActorNodes.Emplace(HLODActorGuid, new FHLODSceneNode());
+		HLODSceneNode->Bounds = HLODActorDesc.GetEditorBounds();
+		HLODSceneNode->HLODActorHandle = FWorldPartitionHandle(InWorldPartition, HLODActorDesc.GetGuid());
+	}
+
+	// Second pass, build a hierarchy now that nodes were all created
+	for (FActorDescContainerCollection::TIterator<AWorldPartitionHLOD> HLODIterator(WorldPartition); HLODIterator; ++HLODIterator)
+	{
+		const FHLODActorDesc& HLODActorDesc = **HLODIterator;
+		const FGuid& HLODActorGuid = HLODActorDesc.GetGuid();
+		TUniquePtr<FHLODSceneNode>& HLODSceneNode = HLODActorNodes.FindChecked(HLODActorGuid);
+
+		for (const FGuid& ChildHLODActorGuid : HLODActorDesc.GetChildHLODActors())
+		{
+			if (TUniquePtr<FHLODSceneNode>* ChildHLODSceneNodePtr = HLODActorNodes.Find(ChildHLODActorGuid))
+			{
+				FHLODSceneNode* ChildHLODSceneNode = ChildHLODSceneNodePtr->Get();
+				ChildHLODSceneNode->ParentHLOD = HLODSceneNode.Get();
+
+				HLODSceneNode->ChildrenHLODs.Add(ChildHLODSceneNode);
+			}
+		}
+	}
+
+	// Cache top level HLOD nodes for a faster iteration during the subsystem tick.
+	for (const auto& [HLODActorGuid, HLODActorNode] : HLODActorNodes)
+	{
+		if (HLODActorNode->ParentHLOD == nullptr)
+		{
+			TopLevelHLODActorNodes.Add(HLODActorNode.Get());
+		}
+	}
+}
+
+FWorldPartitionHLODEditorData::~FWorldPartitionHLODEditorData()
+{
+}
+
+void FWorldPartitionHLODEditorData::UpdateLoadedActorsState()
+{
+	TRACE_CPUPROFILER_EVENT_SCOPE(FWorldPartitionHLODEditorData::UpdateLoadedActorsState);
+
+	// Increment update counter - used to quickly find out if an HLOD actor needs be hidden without having to flag the whole hierarchy
+	LastStateUpdate++;
+
+	struct FBoundsWithVolume
+	{
+		FBox Box;
+		FBox::FReal Volume;
+	};
+
+	TArray<FBoundsWithVolume> LoadedBounds;
+
+	// Gather LoaderAdapter
+	for (UWorldPartitionEditorLoaderAdapter* EditorLoaderAdapter : WorldPartition->GetRegisteredEditorLoaderAdapters())
+	{
+		if (IWorldPartitionActorLoaderInterface::ILoaderAdapter* LoaderAdapter = EditorLoaderAdapter->GetLoaderAdapter())
+		{
+			if (LoaderAdapter->IsLoaded() && LoaderAdapter->GetBoundingBox().IsSet())
+			{
+				const FBox& Box = LoaderAdapter->GetBoundingBox().GetValue();
+				LoadedBounds.Emplace(Box, Box.GetVolume());
+			}
+		}
+	}
+
+	// Gather ActorLoaderInterface
+	FWorldPartitionHelpers::ForEachActorDesc(WorldPartition, [&LoadedBounds](const FWorldPartitionActorDesc* ActorDesc)
+	{
+		if (ActorDesc->GetActorNativeClass()->ImplementsInterface(UWorldPartitionActorLoaderInterface::StaticClass()))
+		{
+			if (AActor* Actor = ActorDesc->GetActor(false))
+			{
+				if (IWorldPartitionActorLoaderInterface::ILoaderAdapter* LoaderAdapter = Cast<IWorldPartitionActorLoaderInterface>(Actor)->GetLoaderAdapter())
+				{
+					if (LoaderAdapter->IsLoaded() && LoaderAdapter->GetBoundingBox().IsSet())
+					{
+						const FBox& Box = ActorDesc->GetEditorBounds();
+						LoadedBounds.Emplace(Box, Box.GetVolume());
+					}
+				}
+			}
+		}
+
+		return true;
+	});
+
+	// Gather Pinned Actors
+	if (WorldPartition->PinnedActors)
+	{
+		WorldPartition->PinnedActors->ForEachReferencedActor([&LoadedBounds](const FWorldPartitionReference& LoadedActor)
+		{
+			if (LoadedActor->GetIsSpatiallyLoaded() && LoadedActor->IsEditorRelevant() && LoadedActor->GetActorIsHLODRelevant())
+			{
+				const FBox& Box = LoadedActor->GetEditorBounds();
+				LoadedBounds.Emplace(Box, Box.GetVolume());
+			}
+		});
+	}
+
+	// Sort Bounds by volume
+	Algo::Sort(LoadedBounds, [](const FBoundsWithVolume& BoxA, const FBoundsWithVolume& BoxB) { return BoxA.Volume > BoxB.Volume; });
+
+	const auto UpdateNodeState = [StateUpdate = LastStateUpdate, &LoadedBounds](FHLODSceneNode* HLODSceneNode)
+	{
+		auto UpdateNodeStateImpl = [StateUpdate, &LoadedBounds](FHLODSceneNode* HLODSceneNode, auto& UpdateNodeStateRef) -> void
+		{
+			FBox HLODSceneNodeBox = HLODSceneNode->Bounds.GetBox();
+
+			const bool bHasAnyIntersectingLoadedRegion = Algo::AnyOf(LoadedBounds, [&HLODSceneNodeBox](const FBoundsWithVolume& Bounds) { return Bounds.Box.Intersect(HLODSceneNodeBox); });
+			if (bHasAnyIntersectingLoadedRegion)
+			{
+				HLODSceneNode->HasIntersectingLoadedRegion = StateUpdate;
+
+				for (FHLODSceneNode* Child : HLODSceneNode->ChildrenHLODs)
+				{
+					UpdateNodeStateRef(Child, UpdateNodeStateRef);
+				}
+			}
+		};
+
+		return UpdateNodeStateImpl(HLODSceneNode, UpdateNodeStateImpl);
+	};
+
+	// Update Nodes, starting from the top level HLODs down to their children
+	for (FHLODSceneNode* HLODSceneNode : TopLevelHLODActorNodes)
+	{
+		UpdateNodeState(HLODSceneNode);
+	}
+}
+
+void FWorldPartitionHLODEditorData::UpdateVisibility(const FVector& InCameraLocation, bool bInForceVisibilityUpdate)
+{
+	TRACE_CPUPROFILER_EVENT_SCOPE(FWorldPartitionHLODEditorData::UpdateVisibility);
+
+	// Early out if no HLOD actor is loaded - none will be visible
+	if (!HLODActorsLoader.IsValid())
+	{
+		return;
+	}
+
+	// For each top level HLOD actor (ex: HLOD2 in a 3 level of HLOD setup)
+	// Determine visibility of each HLOD by a few factors:
+	//  * If source (non-hlod) actors are loaded beneat it - HIDDEN
+	//  * If near culled (using the min visible distance) - HIDDEN
+	// Then if a given HLOD is HIDDEN, perform the same logic for its children.
+	// If the HLOD is VISIBLE, then flag all it's children as HIDDEN
+	for (const auto HLODActorNode : TopLevelHLODActorNodes)
+	{
+		// Recurse from the top level HLOD down to HLOD0
+		HLODActorNode->UpdateVisibility(InCameraLocation, /*bInForceHidden*/false, bInForceVisibilityUpdate, LastStateUpdate);
+	}
+}
+
+void FHLODSceneNode::UpdateVisibility(const FVector& InCameraLocation, bool bInForceHidden, bool bInForceVisibilityUpdate, int32 InLastStateUpdate)
+{
+	bool bNodeShouldBeVisible = !bInForceHidden;
+
+	if (AWorldPartitionHLOD* HLODActor = Cast<AWorldPartitionHLOD>(HLODActorHandle->GetActor()))
+	{
+		// Do not show node if there are loaded actors under it
+		if (bNodeShouldBeVisible)
+		{
+			const bool bHasLoadedActors = HasIntersectingLoadedRegion == InLastStateUpdate;
+			bNodeShouldBeVisible = !bHasLoadedActors;
+		}
+
+		// Perform distance culling
+		if (bNodeShouldBeVisible)
+		{
+			// Do not perform near culling if this HLOD has no child HLOD or none of them is loaded
+			const bool bHasLoadedChildrenHLODs = Algo::AnyOf(ChildrenHLODs, [](FHLODSceneNode* ChildHLODNode) { return ChildHLODNode->HLODActorHandle.IsLoaded(); });
+			if (bHasLoadedChildrenHLODs)
+			{
+				const float DistanceSquared = Bounds.ComputeSquaredDistanceFromBoxToPoint(InCameraLocation);
+				const bool bNearCulled = DistanceSquared < FMath::Square(HLODActor->GetMinVisibleDistance());
+				bNodeShouldBeVisible = !bNearCulled;
+			}
+		}
+
+		if (bInForceVisibilityUpdate || bCachedIsVisible != bNodeShouldBeVisible)
+		{
+			bCachedIsVisible = bNodeShouldBeVisible;
+			HLODActor->SetVisibility(bNodeShouldBeVisible);
+		}
+	}
+
+	const bool bForceHideChildren = bNodeShouldBeVisible || bInForceHidden;
+	for (FHLODSceneNode* Child : ChildrenHLODs)
+	{
+		Child->UpdateVisibility(InCameraLocation, bForceHideChildren, bInForceVisibilityUpdate, InLastStateUpdate);
+	}
+}
+
+void FWorldPartitionHLODEditorData::LoadHLODActors()
+{
+	HLODActorsLoader.Reset(new FLoaderAdapterHLOD(WorldPartition->GetTypedOuter<UWorld>()));
+}
+
+void FWorldPartitionHLODEditorData::UnloadHLODActors()
+{
+	HLODActorsLoader.Reset();
+}
diff --git a/Engine/Source/Editor/WorldPartitionEditor/Private/WorldPartition/HLOD/HLODEditorSubsystem.cpp b/Engine/Source/Editor/WorldPartitionEditor/Private/WorldPartition/HLOD/HLODEditorSubsystem.cpp
new file mode 100644
--- /dev/null
+++ b/Engine/Source/Editor/WorldPartitionEditor/Private/WorldPartition/HLOD/HLODEditorSubsystem.cpp
@@ -0,0 +1,149 @@
+// Copyright Epic Games, Inc. All Rights Reserved.
+
+#include "WorldPartition/HLOD/HLODEditorSubsystem.h"
+
+#include "Editor.h"
+#include "Engine/Engine.h"
+#include "Engine/World.h"
+#include "Subsystems/UnrealEditorSubsystem.h"
+#include "WorldPartition/HLOD/HLODEditorData.h"
+#include "WorldPartition/WorldPartition.h"
+#include "WorldPartition/WorldPartitionSubsystem.h"
+
+
+static TAutoConsoleVariable<bool> CVarHLODInEditorEnabled(
+	TEXT("wp.Editor.HLOD"),
+	false,
+	TEXT("Show World Partition HLODs in the editor."));
+
+
+UWorldPartitionHLODEditorSubsystem::UWorldPartitionHLODEditorSubsystem()
+{
+}
+
+UWorldPartitionHLODEditorSubsystem::~UWorldPartitionHLODEditorSubsystem()
+{
+}
+
+bool UWorldPartitionHLODEditorSubsystem::DoesSupportWorldType(const EWorldType::Type WorldType) const
+{
+	return WorldType == EWorldType::Editor && !IsRunningCommandlet();
+}
+
+void UWorldPartitionHLODEditorSubsystem::Initialize(FSubsystemCollectionBase& Collection)
+{
+	// Ensure the WorldPartitionSubsystem gets created before the HLODEditorSubsystem
+	Collection.InitializeDependency<UWorldPartitionSubsystem>();
+
+	Super::Initialize(Collection);
+
+	bForceHLODStateUpdate = true;
+	bHLODInEditorEnabled = CVarHLODInEditorEnabled.GetValueOnGameThread();
+	
+	GetWorld()->OnWorldPartitionInitialized().AddUObject(this, &UWorldPartitionHLODEditorSubsystem::OnWorldPartitionInitialized);
+	GetWorld()->OnWorldPartitionUninitialized().AddUObject(this, &UWorldPartitionHLODEditorSubsystem::OnWorldPartitionUninitialized);
+
+	GEngine->OnLevelActorListChanged().AddUObject(this, &UWorldPartitionHLODEditorSubsystem::ForceHLODStateUpdate);
+}
+
+void UWorldPartitionHLODEditorSubsystem::Deinitialize()
+{
+	Super::Deinitialize();
+
+	GEngine->OnLevelActorListChanged().RemoveAll(this);
+
+	GetWorld()->OnWorldPartitionInitialized().RemoveAll(this);
+	GetWorld()->OnWorldPartitionUninitialized().RemoveAll(this);
+}
+
+void UWorldPartitionHLODEditorSubsystem::OnWorldPartitionInitialized(UWorldPartition* InWorldPartition)
+{
+	TRACE_CPUPROFILER_EVENT_SCOPE(UWorldPartitionHLODEditorSubsystem::OnWorldPartitionInitialized);
+
+	InWorldPartition->LoaderAdapterStateChanged.AddUObject(this, &UWorldPartitionHLODEditorSubsystem::OnLoaderAdapterStateChanged);
+
+
+	FWorldPartitionHLODEditorData* HLODEditorData = WorldPartitionsHLODEditorData.Emplace(InWorldPartition, new FWorldPartitionHLODEditorData(InWorldPartition));
+
+	if (CVarHLODInEditorEnabled.GetValueOnGameThread())
+	{
+		HLODEditorData->LoadHLODActors();
+		ForceHLODStateUpdate();
+	}
+}
+
+void UWorldPartitionHLODEditorSubsystem::OnWorldPartitionUninitialized(UWorldPartition* InWorldPartition)
+{
+	TRACE_CPUPROFILER_EVENT_SCOPE(UWorldPartitionHLODEditorSubsystem::OnWorldPartitionUninitialized);
+
+	InWorldPartition->LoaderAdapterStateChanged.RemoveAll(this);
+
+	FWorldPartitionHLODEditorData* HLODEditorData = WorldPartitionsHLODEditorData.FindAndRemoveChecked(InWorldPartition);
+	delete HLODEditorData;
+}
+
+void UWorldPartitionHLODEditorSubsystem::OnLoaderAdapterStateChanged(const IWorldPartitionActorLoaderInterface::ILoaderAdapter* LoaderAdapter)
+{
+	TRACE_CPUPROFILER_EVENT_SCOPE(UWorldPartitionHLODEditorSubsystem::OnLoaderAdapterStateChanged);
+
+	ForceHLODStateUpdate();
+}
+
+void UWorldPartitionHLODEditorSubsystem::ForceHLODStateUpdate()
+{
+	bForceHLODStateUpdate = true;
+}
+
+void UWorldPartitionHLODEditorSubsystem::Tick(float DeltaTime)
+{
+	TRACE_CPUPROFILER_EVENT_SCOPE(UWorldPartitionHLODEditorSubsystem::Tick);
+
+	if (CVarHLODInEditorEnabled.GetValueOnGameThread() != bHLODInEditorEnabled)
+	{
+		bHLODInEditorEnabled = CVarHLODInEditorEnabled.GetValueOnGameThread();
+		for (auto& [WorldPartition, HLODEditorData] : WorldPartitionsHLODEditorData)
+		{
+			if (bHLODInEditorEnabled)
+			{
+				HLODEditorData->LoadHLODActors();
+			}
+			else
+			{
+				HLODEditorData->UnloadHLODActors();
+			}
+		}
+	}
+
+	if (bForceHLODStateUpdate)
+	{
+		for (auto& [WorldPartition, HLODEditorData] : WorldPartitionsHLODEditorData)
+		{
+			HLODEditorData->UpdateLoadedActorsState();
+		}
+	}
+
+	UUnrealEditorSubsystem* UnrealEditorSubsystem = GEditor->GetEditorSubsystem<UUnrealEditorSubsystem>();
+	if (UnrealEditorSubsystem)
+	{
+		FVector CameraLocation;
+		FRotator CameraRotation;
+		UnrealEditorSubsystem->GetLevelViewportCameraInfo(CameraLocation, CameraRotation);
+
+		if (bForceHLODStateUpdate || CameraLocation != CachedCameraLocation)
+		{
+			CachedCameraLocation = CameraLocation;			
+
+			for (auto& [WorldPartition, HLODEditorData] : WorldPartitionsHLODEditorData)
+			{
+				HLODEditorData->UpdateVisibility(CameraLocation, bForceHLODStateUpdate);
+			}
+		}
+	}
+
+	bForceHLODStateUpdate = false;
+}
+
+TStatId UWorldPartitionHLODEditorSubsystem::GetStatId() const
+{
+	RETURN_QUICK_DECLARE_CYCLE_STAT(WorldPartitionHLODEditorSubsystem, STATGROUP_Tickables);
+}
diff --git a/Engine/Source/Editor/WorldPartitionEditor/Private/WorldPartition/HLOD/HLODLoaderAdapter.cpp b/Engine/Source/Editor/WorldPartitionEditor/Private/WorldPartition/HLOD/HLODLoaderAdapter.cpp
new file mode 100644
--- /dev/null
+++ b/Engine/Source/Editor/WorldPartitionEditor/Private/WorldPartition/HLOD/HLODLoaderAdapter.cpp
@@ -0,0 +1,58 @@
+// Copyright Epic Games, Inc. All Rights Reserved.
+
+#include "WorldPartition/HLOD/HLODLoaderAdapter.h"
+
+#include "Engine/World.h"
+
+#include "WorldPartition/HLOD/HLODActor.h"
+#include "WorldPartition/HLOD/HLODActorDesc.h"
+#include "WorldPartition/HLOD/HLODLayer.h"
+
+#include "WorldPartition/ActorDescContainerCollection.h"
+#include "WorldPartition/WorldPartition.h"
+#include "WorldPartition/WorldPartitionHandle.h"
+
+
+
+FLoaderAdapterHLOD::FLoaderAdapterHLOD(UWorld* InWorld)
+	: FLoaderAdapterActorList(InWorld)
+{
+	UWorldPartition* WorldPartition = InWorld->GetWorldPartition();
+
+	for (FActorDescContainerCollection::TIterator<AWorldPartitionHLOD> HLODIterator(WorldPartition); HLODIterator; ++HLODIterator)
+	{
+		const FHLODActorDesc& HLODActorDesc = **HLODIterator;
+		if (ShouldLoadHLOD(HLODActorDesc))
+		{
+			Actors.Add(FWorldPartitionHandle(WorldPartition, HLODActorDesc.GetGuid()));
+		}
+	}
+
+	RefreshLoadedState();
+}
+
+// TODO - How to choose which HLODs are loaded ?
+// * Always loaded should be loaded...
+// * Should be a user setting/bookmark setting maybe?
+// * World settings? 
+// * HLOD Layer setting?
+//
+// Current solution is to load always loaded HLOD + any HLOD not built from instancing
+bool FLoaderAdapterHLOD::ShouldLoadHLOD(const FHLODActorDesc& HLODActorDesc)
+{
+	FSoftObjectPath HLODLayerPath(HLODActorDesc.GetSourceHLODLayer());
+	if (UHLODLayer* HLODLayer = Cast<UHLODLayer>(HLODLayerPath.TryLoad()))
+	{
+		if (!HLODActorDesc.GetIsSpatiallyLoaded() || HLODLayer->GetLayerType() != EHLODLayerType::Instancing)
+		{
+			return true;
+		}
+	}
+
+	return false;
+}
+
+bool FLoaderAdapterHLOD::PassActorDescFilter(const FWorldPartitionHandle& ActorHandle) const
+{
+	return ActorHandle.IsValid() && !ActorsToRemove.Contains(ActorHandle);
+}
diff --git a/Engine/Source/Editor/WorldPartitionEditor/Public/WorldPartition/HLOD/HLODEditorData.h b/Engine/Source/Editor/WorldPartitionEditor/Public/WorldPartition/HLOD/HLODEditorData.h
new file mode 100644
--- /dev/null
+++ b/Engine/Source/Editor/WorldPartitionEditor/Public/WorldPartition/HLOD/HLODEditorData.h
@@ -0,0 +1,45 @@
+// Copyright Epic Games, Inc. All Rights Reserved.
+
+#pragma once
+
+#include "WorldPartition/LoaderAdapter/LoaderAdapterActorList.h"
+
+
+class UWorldPartition;
+
+// Represent an HLOD actor in the editor, loaded or not
+struct FHLODSceneNode
+{
+	void UpdateVisibility(const FVector& InCameraLocation, bool bInForceHidden, bool bInForceVisibilityUpdate, int32 InLastStateUpdate);
+
+	FHLODSceneNode* ParentHLOD = nullptr;
+	TArray<FHLODSceneNode*> ChildrenHLODs;
+
+
+	FBoxSphereBounds Bounds;
+	bool bCachedIsVisible = true;
+	FWorldPartitionHandle HLODActorHandle;
+	
+	int32 HasIntersectingLoadedRegion = INDEX_NONE;
+};
+
+// Editor state of HLODs for a given World Partition
+struct FWorldPartitionHLODEditorData
+{
+public:
+	FWorldPartitionHLODEditorData(UWorldPartition* InWorldPartition);
+	~FWorldPartitionHLODEditorData();
+
+	void UpdateLoadedActorsState();
+	void UpdateVisibility(const FVector& InCameraLocation, bool bForceVisibilityUpdate);
+
+	void LoadHLODActors();
+	void UnloadHLODActors();
+
+private:
+	UWorldPartition* WorldPartition;
+	TMap<FGuid, TUniquePtr<FHLODSceneNode>> HLODActorNodes;
+	TArray<FHLODSceneNode*> TopLevelHLODActorNodes;
+	TUniquePtr<FLoaderAdapterActorList> HLODActorsLoader;
+	int32 LastStateUpdate;
+};
diff --git a/Engine/Source/Editor/WorldPartitionEditor/Public/WorldPartition/HLOD/HLODEditorSubsystem.h b/Engine/Source/Editor/WorldPartitionEditor/Public/WorldPartition/HLOD/HLODEditorSubsystem.h
new file mode 100644
--- /dev/null
+++ b/Engine/Source/Editor/WorldPartitionEditor/Public/WorldPartition/HLOD/HLODEditorSubsystem.h
@@ -0,0 +1,59 @@
+// Copyright Epic Games, Inc. All Rights Reserved.
+
+#pragma once
+
+#include "Subsystems/WorldSubsystem.h"
+#include "WorldPartition/WorldPartitionActorLoaderInterface.h"
+#include "WorldPartition/WorldPartitionHandle.h"
+
+#include "HLODEditorSubsystem.generated.h"
+
+
+class AActor;
+class AWorldPartitionHLOD;
+class UWorldPartition;
+struct FWorldPartitionHLODEditorData;
+
+/**
+ * UWorldPartitionHLODEditorSubsystem
+ */
+UCLASS()
+class WORLDPARTITIONEDITOR_API UWorldPartitionHLODEditorSubsystem : public UTickableWorldSubsystem
+{
+	GENERATED_BODY()
+
+public:
+	UWorldPartitionHLODEditorSubsystem();
+	virtual ~UWorldPartitionHLODEditorSubsystem();
+
+	//~ Begin USubsystem Interface.
+	virtual void Initialize(FSubsystemCollectionBase& Collection) override;
+	virtual void Deinitialize() override;
+	//~ End USubsystem Interface.
+
+	//~ Begin UWorldSubsystem Interface.
+	virtual bool DoesSupportWorldType(const EWorldType::Type WorldType) const override;
+	//~ End UWorldSubsystem Interface.
+
+	//~ Begin FTickableGameObject Interface
+	virtual bool IsTickable() const override { return true; }
+	virtual bool IsTickableInEditor() const override { return true; }
+	virtual void Tick(float DeltaTime) override;
+	virtual TStatId GetStatId() const override;
+	//~ End FTickableGameObject Interface
+
+private:
+	void OnWorldPartitionInitialized(UWorldPartition* InWorldPartition);
+	void OnWorldPartitionUninitialized(UWorldPartition* InWorldPartition);
+
+	void OnLoaderAdapterStateChanged(const IWorldPartitionActorLoaderInterface::ILoaderAdapter* LoaderAdapter);
+
+	void ForceHLODStateUpdate();
+	
+private:
+	FVector CachedCameraLocation;
+	bool bForceHLODStateUpdate;
+	bool bHLODInEditorEnabled;
+
+	TMap<TObjectPtr<UWorldPartition>, FWorldPartitionHLODEditorData*> WorldPartitionsHLODEditorData;
+};
diff --git a/Engine/Source/Editor/WorldPartitionEditor/Public/WorldPartition/HLOD/HLODLoaderAdapter.h b/Engine/Source/Editor/WorldPartitionEditor/Public/WorldPartition/HLOD/HLODLoaderAdapter.h
new file mode 100644
--- /dev/null
+++ b/Engine/Source/Editor/WorldPartitionEditor/Public/WorldPartition/HLOD/HLODLoaderAdapter.h
@@ -0,0 +1,18 @@
+// Copyright Epic Games, Inc. All Rights Reserved.
+
+#pragma once
+
+#include "WorldPartition/LoaderAdapter/LoaderAdapterActorList.h"
+
+class FHLODActorDesc;
+
+class FLoaderAdapterHLOD : public FLoaderAdapterActorList
+{
+public:
+	FLoaderAdapterHLOD(UWorld* InWorld);
+
+	virtual bool PassActorDescFilter(const FWorldPartitionHandle& ActorHandle) const override;
+
+private:
+	bool ShouldLoadHLOD(const FHLODActorDesc& HLODActorDesc);	
+};
