Build fix for MacArm64
#rnx

diff --git a/Engine/Source/Editor/WorldPartitionEditor/Private/WorldPartition/HLOD/HLODEditorData.cpp b/Engine/Source/Editor/WorldPartitionEditor/Private/WorldPartition/HLOD/HLODEditorData.cpp
--- a/Engine/Source/Editor/WorldPartitionEditor/Private/WorldPartition/HLOD/HLODEditorData.cpp
+++ b/Engine/Source/Editor/WorldPartitionEditor/Private/WorldPartition/HLOD/HLODEditorData.cpp
@@ -69,6 +69,12 @@ void FWorldPartitionHLODEditorData::UpdateLoadedActorsState()
 
 	struct FBoundsWithVolume
 	{
+		FBoundsWithVolume(const FBox& InBox, const FBox::FReal InVolume)
+			: Box(InBox)
+			, Volume(InVolume)
+		{
+		}
+
 		FBox Box;
 		FBox::FReal Volume;
 	};
