Linux: Fix use-after-stack-free issue reported by ASan.

#rnx

diff --git a/Engine/Source/Editor/WorldPartitionEditor/Private/WorldPartition/HLOD/HLODEditorData.cpp b/Engine/Source/Editor/WorldPartitionEditor/Private/WorldPartition/HLOD/HLODEditorData.cpp
--- a/Engine/Source/Editor/WorldPartitionEditor/Private/WorldPartition/HLOD/HLODEditorData.cpp
+++ b/Engine/Source/Editor/WorldPartitionEditor/Private/WorldPartition/HLOD/HLODEditorData.cpp
@@ -88,7 +88,7 @@ void FWorldPartitionHLODEditorData::UpdateLoadedActorsState()
 		{
 			if (LoaderAdapter->IsLoaded() && LoaderAdapter->GetBoundingBox().IsSet())
 			{
-				const FBox& Box = LoaderAdapter->GetBoundingBox().GetValue();
+				const FBox Box = LoaderAdapter->GetBoundingBox().GetValue();
 				LoadedBounds.Emplace(Box, Box.GetVolume());
 			}
 		}
@@ -105,7 +105,7 @@ void FWorldPartitionHLODEditorData::UpdateLoadedActorsState()
 				{
 					if (LoaderAdapter->IsLoaded() && LoaderAdapter->GetBoundingBox().IsSet())
 					{
-						const FBox& Box = ActorDesc->GetEditorBounds();
+						const FBox Box = ActorDesc->GetEditorBounds();
 						LoadedBounds.Emplace(Box, Box.GetVolume());
 					}
 				}
@@ -122,7 +122,7 @@ void FWorldPartitionHLODEditorData::UpdateLoadedActorsState()
 		{
 			if (LoadedActor->GetIsSpatiallyLoaded() && LoadedActor->IsEditorRelevant() && LoadedActor->GetActorIsHLODRelevant())
 			{
-				const FBox& Box = LoadedActor->GetEditorBounds();
+				const FBox Box = LoadedActor->GetEditorBounds();
 				LoadedBounds.Emplace(Box, Box.GetVolume());
 			}
 		});
