Fixed blurry WP editor minimap
* Provide the correct visible bounds of the minimap texture to RequestVirtualTextureTiles()
* Would lead to incorrect tiles being loaded when zooming, leading to a blurry minimap.

diff --git a/Engine/Source/Editor/WorldPartitionEditor/Private/WorldPartition/SWorldPartitionEditorGrid2D.cpp b/Engine/Source/Editor/WorldPartitionEditor/Private/WorldPartition/SWorldPartitionEditorGrid2D.cpp
--- a/Engine/Source/Editor/WorldPartitionEditor/Private/WorldPartition/SWorldPartitionEditorGrid2D.cpp
+++ b/Engine/Source/Editor/WorldPartitionEditor/Private/WorldPartition/SWorldPartitionEditorGrid2D.cpp
@@ -1934,7 +1934,7 @@ int32 SWorldPartitionEditorGrid2D::PaintMinimap(const FGeometry& AllottedGeometr
 		);
 
 		const FPaintGeometry WorldImageGeometry = AllottedGeometry.ToPaintGeometry(
-			MinimapBounds.Max - MinimapBounds.Min,
+			MinimapBounds.GetSize(),
 			FSlateLayoutTransform(MinimapBounds.Min)
 		);
 
@@ -2027,8 +2027,8 @@ int32 SWorldPartitionEditorGrid2D::PaintMinimap(const FGeometry& AllottedGeometr
 		{
 			FVirtualTexture2DResource* VTResource = static_cast<FVirtualTexture2DResource*>(Texture2D->GetResource());
 			const FVector2D ViewportSize = AllottedGeometry.GetLocalSize();
-			const FVector2D ScreenSpaceSize = WorldImageGeometry.GetLocalSize();
-			const FVector2D ViewportPositon = AllottedGeometry.GetAbsolutePosition() - WorldImageGeometry.GetAccumulatedRenderTransform().GetTranslation();
+			const FVector2D ScreenSpaceSize = MinimapBounds.GetSize();
+			const FVector2D ViewportPosition = MinimapBounds.Min;
 
 			FBox2D UVRegion = WorldMiniMapBrush.GetUVRegion();
 			const FVector2D UV0 = UVRegion.Min;
@@ -2040,13 +2040,13 @@ int32 SWorldPartitionEditorGrid2D::PaintMinimap(const FGeometry& AllottedGeometr
 			UE::RenderCommandPipe::FSyncScope SyncScope;
 
 			ENQUEUE_RENDER_COMMAND(MakeTilesResident)(
-				[InFeatureLevel, VTResource, ScreenSpaceSize, ViewportPositon, ViewportSize, UV0, UV1, MipLevel](FRHICommandListImmediate& RHICmdList)
+				[InFeatureLevel, VTResource, ScreenSpaceSize, ViewportPosition, ViewportSize, UV0, UV1, MipLevel](FRHICommandListImmediate& RHICmdList)
 			{
 				// AcquireAllocatedVT() must happen on render thread
 				IAllocatedVirtualTexture* AllocatedVT = VTResource->AcquireAllocatedVT();
 
 				IRendererModule& RenderModule = GetRendererModule();
-				RenderModule.RequestVirtualTextureTiles(AllocatedVT, ScreenSpaceSize, ViewportPositon, ViewportSize, UV0, UV1, MipLevel);
+				RenderModule.RequestVirtualTextureTiles(AllocatedVT, ScreenSpaceSize, ViewportPosition, ViewportSize, UV0, UV1, MipLevel);
 				RenderModule.LoadPendingVirtualTextureTiles(RHICmdList, InFeatureLevel);
 			});
 		}
