ActorFilter Customization: Check if property is editable

#rnx

diff --git a/Engine/Source/Editor/WorldPartitionEditor/Private/WorldPartition/Filter/WorldPartitionActorFilterPropertyTypeCustomization.cpp b/Engine/Source/Editor/WorldPartitionEditor/Private/WorldPartition/Filter/WorldPartitionActorFilterPropertyTypeCustomization.cpp
--- a/Engine/Source/Editor/WorldPartitionEditor/Private/WorldPartition/Filter/WorldPartitionActorFilterPropertyTypeCustomization.cpp
+++ b/Engine/Source/Editor/WorldPartitionEditor/Private/WorldPartition/Filter/WorldPartitionActorFilterPropertyTypeCustomization.cpp
@@ -100,6 +100,7 @@ void FWorldPartitionActorFilterPropertyTypeCustomization::CustomizeHeader(TShare
 		.MinDesiredWidth(2000.f) // Set very large as it is the only way for the Outliner to expand to take all the space
 		.MinDesiredHeight(150.f)
 		.MaxDesiredHeight(300.f)
+		.IsEnabled(StructPropertyHandle->IsEditable())
 		[
 			SNew(SVerticalBox)
 			+ SVerticalBox::Slot()
