World Partition - HLOD in Editor
* Pinning Level Instances will now properly hide the HLOD(s) under it
* Removed ILoaderAdapter::ForEachReferencedActor()

diff --git a/Engine/Source/Editor/WorldPartitionEditor/Private/WorldPartition/HLOD/HLODEditorData.cpp b/Engine/Source/Editor/WorldPartitionEditor/Private/WorldPartition/HLOD/HLODEditorData.cpp
--- a/Engine/Source/Editor/WorldPartitionEditor/Private/WorldPartition/HLOD/HLODEditorData.cpp
+++ b/Engine/Source/Editor/WorldPartitionEditor/Private/WorldPartition/HLOD/HLODEditorData.cpp
@@ -56,29 +56,62 @@ FWorldPartitionHLODEditorData::FWorldPartitionHLODEditorData(UWorldPartition* In
 	}
 }
 
-FWorldPartitionHLODEditorData::~FWorldPartitionHLODEditorData()
+struct FBoundsWithVolume
 {
-}
+	FBoundsWithVolume(const FBox& InBox)
+		: Box(InBox)
+		, Volume(Box.GetVolume())
+	{
+	}
 
-void FWorldPartitionHLODEditorData::UpdateLoadedActorsState()
+	FBox Box;
+	FBox::FReal Volume;
+};
+
+// Gather Pinned Actors bounds - also include references & contained actors if the pinned actor is a container.
+static void GatherLoadedActorsBounds(TArray<FBoundsWithVolume>& OutLoadedBounds, const FWorldPartitionActorDesc* InActorDesc, const UActorDescContainer* InContainer, const TOptional<FTransform>& InContainerTransform = TOptional<FTransform>())
 {
-	TRACE_CPUPROFILER_EVENT_SCOPE(FWorldPartitionHLODEditorData::UpdateLoadedActorsState);
+	if (InActorDesc)
+	{
+		// The actor itself - Include the actor bounds only if HLOD relevant
+		if (InActorDesc->GetIsSpatiallyLoaded() && InActorDesc->IsEditorRelevant() && InActorDesc->GetActorIsHLODRelevant())
+		{
+			const FBox Box = InContainerTransform.IsSet() ? InActorDesc->GetEditorBounds().TransformBy(InContainerTransform.GetValue()) : InActorDesc->GetEditorBounds();
+			OutLoadedBounds.Emplace(Box);
+		}
 
-	// Increment update counter - used to quickly find out if an HLOD actor needs be hidden without having to flag the whole hierarchy
-	LastStateUpdate++;
+		// Test its references
+		for (const FGuid& ReferenceGuid : InActorDesc->GetReferences())
+		{
+			if (const FWorldPartitionActorDesc* ReferenceActorDesc = InContainer->GetActorDesc(ReferenceGuid))
+			{
+				GatherLoadedActorsBounds(OutLoadedBounds, ReferenceActorDesc, InContainer, InContainerTransform);
+			}
+		}
 
-	struct FBoundsWithVolume
-	{
-		FBoundsWithVolume(const FBox& InBox, const FBox::FReal InVolume)
-			: Box(InBox)
-			, Volume(InVolume)
+		// If it's a container, test the contained actors
+		if (InActorDesc->IsContainerInstance())
 		{
+			FWorldPartitionActorDesc::FContainerInstance ContainerInstance;
+			if (InActorDesc->GetContainerInstance(ContainerInstance))
+			{
+				FTransform ContainerWorldSpaceTransform = InContainerTransform.IsSet() ? ContainerInstance.Transform * InContainerTransform.GetValue() : ContainerInstance.Transform;
+				for (FActorDescList::TConstIterator<> ActorDescIt(ContainerInstance.Container); ActorDescIt; ++ActorDescIt)
+				{
+					GatherLoadedActorsBounds(OutLoadedBounds, *ActorDescIt, ContainerInstance.Container, ContainerWorldSpaceTransform);
+				}
+			}
 		}
+	}
+}
 
-		FBox Box;
-		FBox::FReal Volume;
-	};
+void FWorldPartitionHLODEditorData::UpdateLoadedActorsState()
+{
+	TRACE_CPUPROFILER_EVENT_SCOPE(FWorldPartitionHLODEditorData::UpdateLoadedActorsState);
 
+	// Increment update counter - used to quickly find out if an HLOD actor needs be hidden without having to flag the whole hierarchy
+	LastStateUpdate++;
+	
 	TArray<FBoundsWithVolume> LoadedBounds;
 
 	// Gather LoaderAdapter
@@ -88,8 +121,7 @@ void FWorldPartitionHLODEditorData::UpdateLoadedActorsState()
 		{
 			if (LoaderAdapter->IsLoaded() && LoaderAdapter->GetBoundingBox().IsSet())
 			{
-				const FBox Box = LoaderAdapter->GetBoundingBox().GetValue();
-				LoadedBounds.Emplace(Box, Box.GetVolume());
+				LoadedBounds.Emplace(LoaderAdapter->GetBoundingBox().GetValue());
 			}
 		}
 	}
@@ -105,8 +137,7 @@ void FWorldPartitionHLODEditorData::UpdateLoadedActorsState()
 				{
 					if (LoaderAdapter->IsLoaded() && LoaderAdapter->GetBoundingBox().IsSet())
 					{
-						const FBox Box = ActorDesc->GetEditorBounds();
-						LoadedBounds.Emplace(Box, Box.GetVolume());
+						LoadedBounds.Emplace(ActorDesc->GetEditorBounds());
 					}
 				}
 			}
@@ -114,18 +145,17 @@ void FWorldPartitionHLODEditorData::UpdateLoadedActorsState()
 
 		return true;
 	});
-
-	// Gather Pinned Actors
+	
+	// Gather Pinned Actors bounds
 	if (WorldPartition->PinnedActors)
 	{
-		WorldPartition->PinnedActors->ForEachReferencedActor([&LoadedBounds](const FWorldPartitionReference& LoadedActor)
+		for (const FWorldPartitionHandle& PinnedActor : WorldPartition->PinnedActors->GetActors())
 		{
-			if (LoadedActor->GetIsSpatiallyLoaded() && LoadedActor->IsEditorRelevant() && LoadedActor->GetActorIsHLODRelevant())
+			if (PinnedActor.IsValid())
 			{
-				const FBox Box = LoadedActor->GetEditorBounds();
-				LoadedBounds.Emplace(Box, Box.GetVolume());
+				GatherLoadedActorsBounds(LoadedBounds, *PinnedActor, WorldPartition->GetActorDescContainer());
 			}
-		});
+		}
 	}
 
 	// Sort Bounds by volume
diff --git a/Engine/Source/Editor/WorldPartitionEditor/Private/WorldPartition/HLOD/HLODEditorSubsystem.cpp b/Engine/Source/Editor/WorldPartitionEditor/Private/WorldPartition/HLOD/HLODEditorSubsystem.cpp
--- a/Engine/Source/Editor/WorldPartitionEditor/Private/WorldPartition/HLOD/HLODEditorSubsystem.cpp
+++ b/Engine/Source/Editor/WorldPartitionEditor/Private/WorldPartition/HLOD/HLODEditorSubsystem.cpp
@@ -60,11 +60,15 @@ void UWorldPartitionHLODEditorSubsystem::OnWorldPartitionInitialized(UWorldParti
 {
 	TRACE_CPUPROFILER_EVENT_SCOPE(UWorldPartitionHLODEditorSubsystem::OnWorldPartitionInitialized);
 
-	InWorldPartition->LoaderAdapterStateChanged.AddUObject(this, &UWorldPartitionHLODEditorSubsystem::OnLoaderAdapterStateChanged);
+	UActorDescContainer* ActorDescContainer = InWorldPartition->GetActorDescContainer();
+	if (!ActorDescContainer || ActorDescContainer->IsTemplateContainer())
+	{
+		return;
+	}
 
+	InWorldPartition->LoaderAdapterStateChanged.AddUObject(this, &UWorldPartitionHLODEditorSubsystem::OnLoaderAdapterStateChanged);
 
 	FWorldPartitionHLODEditorData* HLODEditorData = WorldPartitionsHLODEditorData.Emplace(InWorldPartition, new FWorldPartitionHLODEditorData(InWorldPartition));
-
 	if (CVarHLODInEditorEnabled.GetValueOnGameThread())
 	{
 		HLODEditorData->LoadHLODActors();
@@ -76,6 +80,12 @@ void UWorldPartitionHLODEditorSubsystem::OnWorldPartitionUninitialized(UWorldPar
 {
 	TRACE_CPUPROFILER_EVENT_SCOPE(UWorldPartitionHLODEditorSubsystem::OnWorldPartitionUninitialized);
 
+	UActorDescContainer* ActorDescContainer = InWorldPartition->GetActorDescContainer();
+	if (!ActorDescContainer || ActorDescContainer->IsTemplateContainer())
+	{
+		return;
+	}
+
 	InWorldPartition->LoaderAdapterStateChanged.RemoveAll(this);
 
 	FWorldPartitionHLODEditorData* HLODEditorData = WorldPartitionsHLODEditorData.FindAndRemoveChecked(InWorldPartition);
diff --git a/Engine/Source/Editor/WorldPartitionEditor/Public/WorldPartition/HLOD/HLODEditorData.h b/Engine/Source/Editor/WorldPartitionEditor/Public/WorldPartition/HLOD/HLODEditorData.h
--- a/Engine/Source/Editor/WorldPartitionEditor/Public/WorldPartition/HLOD/HLODEditorData.h
+++ b/Engine/Source/Editor/WorldPartitionEditor/Public/WorldPartition/HLOD/HLODEditorData.h
@@ -28,8 +28,7 @@ struct FWorldPartitionHLODEditorData
 {
 public:
 	FWorldPartitionHLODEditorData(UWorldPartition* InWorldPartition);
-	~FWorldPartitionHLODEditorData();
-
+	
 	void UpdateLoadedActorsState();
 	void UpdateVisibility(const FVector& InCameraLocation, bool bForceVisibilityUpdate);
 
diff --git a/Engine/Source/Runtime/Engine/Private/WorldPartition/WorldPartitionActorLoaderInterface.cpp b/Engine/Source/Runtime/Engine/Private/WorldPartition/WorldPartitionActorLoaderInterface.cpp
--- a/Engine/Source/Runtime/Engine/Private/WorldPartition/WorldPartitionActorLoaderInterface.cpp
+++ b/Engine/Source/Runtime/Engine/Private/WorldPartition/WorldPartitionActorLoaderInterface.cpp
@@ -240,23 +240,6 @@ void IWorldPartitionActorLoaderInterface::ILoaderAdapter::OnActorDescContainerUn
 	}
 }
 
-void IWorldPartitionActorLoaderInterface::ILoaderAdapter::ForEachReferencedActor(TFunctionRef<void(const FWorldPartitionReference&)> InFunc) const
-{
-	for (const auto& [Container, ActorReferenceMap] : ContainerActorReferences)
-	{
-		if (Container.IsValid())
-		{
-			for (const auto& [ReferenceGuid, ReferenceMap] : ActorReferenceMap)
-			{
-				for (const auto& [Guid, ActorReference] : ReferenceMap)
-				{
-					InFunc(ActorReference);
-				}
-			}
-		}
-	}
-}
-
 bool IWorldPartitionActorLoaderInterface::ILoaderAdapter::ShouldActorBeLoaded(const FWorldPartitionHandle& Actor) const
 {
 	check(Actor.IsValid());
diff --git a/Engine/Source/Runtime/Engine/Public/WorldPartition/LoaderAdapter/LoaderAdapterList.h b/Engine/Source/Runtime/Engine/Public/WorldPartition/LoaderAdapter/LoaderAdapterList.h
--- a/Engine/Source/Runtime/Engine/Public/WorldPartition/LoaderAdapter/LoaderAdapterList.h
+++ b/Engine/Source/Runtime/Engine/Public/WorldPartition/LoaderAdapter/LoaderAdapterList.h
@@ -12,6 +12,8 @@ public:
 	ENGINE_API FLoaderAdapterList(UWorld* InWorld);
 	virtual ~FLoaderAdapterList() {}
 
+	const TSet<FWorldPartitionHandle>& GetActors() const { return Actors; }
+
 protected:
 	//~ Begin ILoaderAdapter interface
 	ENGINE_API virtual void ForEachActor(TFunctionRef<void(const FWorldPartitionHandle&)> InOperation) const override;
diff --git a/Engine/Source/Runtime/Engine/Public/WorldPartition/WorldPartitionActorLoaderInterface.h b/Engine/Source/Runtime/Engine/Public/WorldPartition/WorldPartitionActorLoaderInterface.h
--- a/Engine/Source/Runtime/Engine/Public/WorldPartition/WorldPartitionActorLoaderInterface.h
+++ b/Engine/Source/Runtime/Engine/Public/WorldPartition/WorldPartitionActorLoaderInterface.h
@@ -48,8 +48,6 @@ public:
 		ENGINE_API void OnActorDescContainerInitialize(UActorDescContainer* Container);
 		ENGINE_API void OnActorDescContainerUninitialize(UActorDescContainer* Container);
 
-		ENGINE_API void ForEachReferencedActor(TFunctionRef<void(const FWorldPartitionReference&)>) const;
-
 	protected:
 		// Private interface
 		virtual void ForEachActor(TFunctionRef<void(const FWorldPartitionHandle&)> InOperation) const =0;
