Removed AssetTypeActions for HLOD Layer & Data Layer assets
* Rely on UAssetDefinitions instead
* Add to correct "World" category

diff --git a/Engine/Plugins/Editor/EngineAssetDefinitions/Source/Private/AssetDefinition_DataLayer.h b/Engine/Plugins/Editor/EngineAssetDefinitions/Source/Private/AssetDefinition_DataLayer.h
--- a/Engine/Plugins/Editor/EngineAssetDefinitions/Source/Private/AssetDefinition_DataLayer.h
+++ b/Engine/Plugins/Editor/EngineAssetDefinitions/Source/Private/AssetDefinition_DataLayer.h
@@ -4,6 +4,7 @@
 
 #include "WorldPartition/DataLayer/DataLayerAsset.h"
 #include "AssetDefinitionDefault.h"
+#include "Misc/AssetCategoryPath.h"
 
 #include "AssetDefinition_DataLayer.generated.h"
 
@@ -17,5 +18,10 @@ public:
 	virtual FText GetAssetDisplayName() const override { return NSLOCTEXT("AssetTypeActions", "AssetTypeActions_DataLayer", "Data Layer"); }
 	virtual FLinearColor GetAssetColor() const override { return FLinearColor(FColor(12, 65, 12)); }
 	virtual TSoftClassPtr<UObject> GetAssetClass() const override { return UDataLayerAsset::StaticClass(); }
+	virtual TConstArrayView<FAssetCategoryPath> GetAssetCategories() const override
+	{
+		static const auto Categories = { EAssetCategoryPaths::World };
+		return Categories;
+	}
 	// UAssetDefinition End
 };
diff --git a/Engine/Plugins/Editor/EngineAssetDefinitions/Source/Private/AssetDefinition_HLODLayer.h b/Engine/Plugins/Editor/EngineAssetDefinitions/Source/Private/AssetDefinition_HLODLayer.h
new file mode 100644
--- /dev/null
+++ b/Engine/Plugins/Editor/EngineAssetDefinitions/Source/Private/AssetDefinition_HLODLayer.h
@@ -0,0 +1,27 @@
+// Copyright Epic Games, Inc. All Rights Reserved.
+
+#pragma once
+
+#include "WorldPartition/HLOD/HLODLayer.h"
+#include "AssetDefinitionDefault.h"
+#include "Misc/AssetCategoryPath.h"
+
+#include "AssetDefinition_HLODLayer.generated.h"
+
+UCLASS()
+class UAssetDefinition_HLODLayer : public UAssetDefinitionDefault
+{
+	GENERATED_BODY()
+
+public:
+	// UAssetDefinition Begin
+	virtual FText GetAssetDisplayName() const override { return NSLOCTEXT("AssetTypeActions", "AssetTypeActions_HLODLayer", "HLOD Layer"); }
+	virtual FLinearColor GetAssetColor() const override { return FLinearColor(FColor(0, 200, 200)); }
+	virtual TSoftClassPtr<UObject> GetAssetClass() const override { return UHLODLayer::StaticClass(); }
+	virtual TConstArrayView<FAssetCategoryPath> GetAssetCategories() const override
+	{
+		static const auto Categories = { EAssetCategoryPaths::World };
+		return Categories;
+	}
+	// UAssetDefinition End
+};
diff --git a/Engine/Source/Editor/AssetDefinition/Private/AssetDefinition.cpp b/Engine/Source/Editor/AssetDefinition/Private/AssetDefinition.cpp
--- a/Engine/Source/Editor/AssetDefinition/Private/AssetDefinition.cpp
+++ b/Engine/Source/Editor/AssetDefinition/Private/AssetDefinition.cpp
@@ -23,6 +23,7 @@ FAssetCategoryPath EAssetCategoryPaths::Input(LOCTEXT("Input", "Input"));
 FAssetCategoryPath EAssetCategoryPaths::FX(LOCTEXT("FX", "FX"));
 FAssetCategoryPath EAssetCategoryPaths::Cinematics(LOCTEXT("Cinematics", "Cinematics"));
 FAssetCategoryPath EAssetCategoryPaths::Media(LOCTEXT("Media", "Media"));
+FAssetCategoryPath EAssetCategoryPaths::World(LOCTEXT("World", "World"));
 
 FAssetCategoryPath::FAssetCategoryPath(const FText& InCategory)
 {
diff --git a/Engine/Source/Editor/AssetDefinition/Public/AssetDefinition.h b/Engine/Source/Editor/AssetDefinition/Public/AssetDefinition.h
--- a/Engine/Source/Editor/AssetDefinition/Public/AssetDefinition.h
+++ b/Engine/Source/Editor/AssetDefinition/Public/AssetDefinition.h
@@ -278,6 +278,7 @@ struct ASSETDEFINITION_API EAssetCategoryPaths
 	static FAssetCategoryPath Physics;
 	static FAssetCategoryPath Texture;
 	static FAssetCategoryPath UI;
+	static FAssetCategoryPath World;
 };
 
 struct FAssetOpenSupport
diff --git a/Engine/Source/Editor/WorldPartitionEditor/Private/WorldPartition/DataLayers/DataLayerAssetTypeActions.cpp b/Engine/Source/Editor/WorldPartitionEditor/Private/WorldPartition/DataLayers/DataLayerAssetTypeActions.cpp
deleted file mode 100644
diff --git a/Engine/Source/Editor/WorldPartitionEditor/Private/WorldPartition/DataLayers/DataLayerAssetTypeActions.h b/Engine/Source/Editor/WorldPartitionEditor/Private/WorldPartition/DataLayers/DataLayerAssetTypeActions.h
deleted file mode 100644
diff --git a/Engine/Source/Editor/WorldPartitionEditor/Private/WorldPartition/HLOD/HLODLayerAssetTypeActions.cpp b/Engine/Source/Editor/WorldPartitionEditor/Private/WorldPartition/HLOD/HLODLayerAssetTypeActions.cpp
deleted file mode 100644
diff --git a/Engine/Source/Editor/WorldPartitionEditor/Private/WorldPartition/HLOD/HLODLayerAssetTypeActions.h b/Engine/Source/Editor/WorldPartitionEditor/Private/WorldPartition/HLOD/HLODLayerAssetTypeActions.h
deleted file mode 100644
diff --git a/Engine/Source/Editor/WorldPartitionEditor/Private/WorldPartitionEditorModule.cpp b/Engine/Source/Editor/WorldPartitionEditor/Private/WorldPartitionEditorModule.cpp
--- a/Engine/Source/Editor/WorldPartitionEditor/Private/WorldPartitionEditorModule.cpp
+++ b/Engine/Source/Editor/WorldPartitionEditor/Private/WorldPartitionEditorModule.cpp
@@ -7,8 +7,6 @@
 #include "WorldPartition/WorldPartitionMiniMapBuilder.h"
 #include "WorldPartition/WorldPartitionLandscapeSplineMeshesBuilder.h"
 #include "WorldPartition/WorldPartitionActorLoaderInterface.h"
-#include "WorldPartition/HLOD/HLODLayerAssetTypeActions.h"
-#include "WorldPartition/DataLayers/DataLayerAssetTypeActions.h"
 #include "WorldPartition/SWorldPartitionEditor.h"
 #include "WorldPartition/SWorldPartitionEditorGridSpatialHash.h"
 #include "WorldPartition/Customizations/WorldPartitionDetailsCustomization.h"
@@ -173,12 +171,6 @@ void FWorldPartitionEditorModule::StartupModule()
 	
 	UToolMenus::RegisterStartupCallback(FSimpleMulticastDelegate::FDelegate::CreateRaw(this, &FWorldPartitionEditorModule::RegisterMenus));	
 
-	IAssetTools& AssetTools = FModuleManager::LoadModuleChecked<FAssetToolsModule>("AssetTools").Get();
-	HLODLayerAssetTypeActions = MakeShareable(new FHLODLayerAssetTypeActions);
-	AssetTools.RegisterAssetTypeActions(HLODLayerAssetTypeActions.ToSharedRef());
-	DataLayerAssetTypeActions = MakeShareable(new FDataLayerAssetTypeActions);
-	AssetTools.RegisterAssetTypeActions(DataLayerAssetTypeActions.ToSharedRef());
-
 	FPropertyEditorModule& PropertyEditor = FModuleManager::LoadModuleChecked<FPropertyEditorModule>("PropertyEditor");
 	PropertyEditor.RegisterCustomClassLayout("WorldPartition", FOnGetDetailCustomizationInstance::CreateStatic(&FWorldPartitionDetails::MakeInstance));
 	PropertyEditor.RegisterCustomClassLayout("WorldPartitionRuntimeSpatialHash", FOnGetDetailCustomizationInstance::CreateStatic(&FWorldPartitionRuntimeSpatialHashDetails::MakeInstance));
@@ -215,26 +207,6 @@ void FWorldPartitionEditorModule::ShutdownModule()
 		UToolMenus::UnregisterOwner(this);
 	}
 
-	if (FModuleManager::Get().IsModuleLoaded("AssetTools"))
-	{
-		IAssetTools& AssetTools = FModuleManager::GetModuleChecked<FAssetToolsModule>("AssetTools").Get();
-
-		// Unregister the HLODLayer asset type actions
-		if (HLODLayerAssetTypeActions.IsValid())
-		{
-			AssetTools.UnregisterAssetTypeActions(HLODLayerAssetTypeActions.ToSharedRef());
-		}
-
-		// Unregister the DataLayer asset type actions
-		if (DataLayerAssetTypeActions.IsValid())
-		{
-			AssetTools.UnregisterAssetTypeActions(DataLayerAssetTypeActions.ToSharedRef());
-		}
-	}
-
-	HLODLayerAssetTypeActions.Reset();
-	DataLayerAssetTypeActions.Reset();
-
 	if (FModuleManager::Get().IsModuleLoaded("PropertyEditor"))
 	{
 		FPropertyEditorModule& PropertyEditor = FModuleManager::GetModuleChecked<FPropertyEditorModule>("PropertyEditor");
diff --git a/Engine/Source/Editor/WorldPartitionEditor/Public/WorldPartitionEditorModule.h b/Engine/Source/Editor/WorldPartitionEditor/Public/WorldPartitionEditorModule.h
--- a/Engine/Source/Editor/WorldPartitionEditor/Public/WorldPartitionEditorModule.h
+++ b/Engine/Source/Editor/WorldPartitionEditor/Public/WorldPartitionEditorModule.h
@@ -149,11 +149,7 @@ private:
 
 	FDelegateHandle LevelEditorExtenderDelegateHandle;
 
-	TSharedPtr<class FHLODLayerAssetTypeActions> HLODLayerAssetTypeActions;
-	TSharedPtr<class FDataLayerAssetTypeActions> DataLayerAssetTypeActions;
-
 	TWeakPtr<SDockTab> WorldPartitionTab;
-
 	TWeakPtr<SDockTab> ContentBundleTab;
 
 	TWeakPtr<SContentBundleBrowser> ContentBundleBrowser;
