World Partition - HLOD in Editor
* Added World & User settings to control whether HLOD are shown in editor
* Added HLOD entry in "Show" viewport menu, exposing the following user settings
-> Show HLODs (enabled by default)
-> Show HLODs over loaded regions (disabled by default)
-> Min Draw Distance (distance at which HLOD starts to be rendered in editor, 128m by default)
-> Max Draw Distance (infinity by default)
* Added "Load Region from Selection" to viewport context menu (on right-click)
* Added World category to Scene Outliner filter, currently only contains "World Partition HLOD"
* Hid a few categories from World Partition HLOD (when showing in Details panel)
* Feature is still disabled globally by the wp.Editor.HLOD.AllowShowingHLODsInEditor cvar

diff --git a/Engine/Source/Editor/EditorStyle/Private/StarshipStyle.cpp b/Engine/Source/Editor/EditorStyle/Private/StarshipStyle.cpp
--- a/Engine/Source/Editor/EditorStyle/Private/StarshipStyle.cpp
+++ b/Engine/Source/Editor/EditorStyle/Private/StarshipStyle.cpp
@@ -4276,7 +4276,7 @@ void FStarshipEditorStyle::FStyle::SetupLevelEditorStyle()
 			Set( "LevelEditor.Tabs.Toolbar",                new IMAGE_BRUSH("Icons/icon_tab_Toolbars_16x", Icon16x16 ) );
 
 			Set( "LevelEditor.Tabs.Viewports",              new IMAGE_BRUSH_SVG("Starship/Common/Viewports", Icon16x16) );
-			Set( "LevelEditor.Tabs.HLOD",                   new IMAGE_BRUSH_SVG("Starship/Common/HierarchicalLODOutliner", Icon16x16) );
+			Set( "LevelEditor.Tabs.HLOD",                   new IMAGE_BRUSH_SVG("Starship/Common/HierarchicalLOD", Icon16x16) );
 			Set( "LevelEditor.Tabs.Debug",                  new IMAGE_BRUSH_SVG("Starship/Common/Bug", Icon16x16) );
 
 			Set( "LevelEditor.Audit",                  		new IMAGE_BRUSH_SVG("Starship/Common/AssetAudit", Icon16x16) );
@@ -4693,6 +4693,7 @@ void FStarshipEditorStyle::FStyle::SetupLevelEditorStyle()
 		Set("ShowFlagsMenu.SubMenu.Layers", new IMAGE_BRUSH_SVG("Starship/Common/Layers", Icon16x16));
 		Set("ShowFlagsMenu.SubMenu.FoliageTypes", new IMAGE_BRUSH_SVG("Starship/Common/FoliageTypes", Icon16x16));
 		Set("ShowFlagsMenu.SubMenu.Sprites", new IMAGE_BRUSH_SVG("Starship/Common/Sprite", Icon16x16));
+		Set("ShowFlagsMenu.SubMenu.HLODs", new IMAGE_BRUSH_SVG("Starship/Common/HierarchicalLOD", Icon16x16));
 	}
 #endif // WITH_EDITOR || (IS_PROGRAM && WITH_UNREAL_DEVELOPER_TOOLS)
 
diff --git a/Engine/Source/Editor/LevelEditor/LevelEditor.Build.cs b/Engine/Source/Editor/LevelEditor/LevelEditor.Build.cs
--- a/Engine/Source/Editor/LevelEditor/LevelEditor.Build.cs
+++ b/Engine/Source/Editor/LevelEditor/LevelEditor.Build.cs
@@ -89,7 +89,8 @@ public class LevelEditor : ModuleRules
 				"UnsavedAssetsTracker",
 				"UncontrolledChangelists",
 				"RenderCore",
-				"DeveloperSettings"
+				"DeveloperSettings",
+				"WorldPartitionEditor",
 			}
 		);
 
@@ -103,7 +104,6 @@ public class LevelEditor : ModuleRules
 				"AutomationWindow",
 				"Layers",
 				"WorldBrowser",
-				"WorldPartitionEditor",
 				"AssetTools",
 				"WorkspaceMenuStructure",
 				"NewLevelDialog",
diff --git a/Engine/Source/Editor/LevelEditor/Private/SLevelViewportToolBar.cpp b/Engine/Source/Editor/LevelEditor/Private/SLevelViewportToolBar.cpp
--- a/Engine/Source/Editor/LevelEditor/Private/SLevelViewportToolBar.cpp
+++ b/Engine/Source/Editor/LevelEditor/Private/SLevelViewportToolBar.cpp
@@ -14,6 +14,7 @@
 #include "Widgets/Layout/SBorder.h"
 #include "Widgets/Layout/SBox.h"
 #include "Widgets/Input/SSpinBox.h"
+#include "Widgets/Input/SSlider.h"
 #include "Styling/AppStyle.h"
 #include "Camera/CameraActor.h"
 #include "Misc/ConfigCacheIni.h"
@@ -47,11 +48,13 @@
 #include "SEditorViewportViewMenuContext.h"
 #include "Bookmarks/IBookmarkTypeTools.h"
 #include "ToolMenu.h"
+#include "WorldPartitionEditorModule.h"
 #include "WorldPartition/WorldPartitionSubsystem.h"
 #include "SLevelViewport.h"
 #include "SortHelper.h"
 #include "Interfaces/IMainFrameModule.h"
 #include "SCommonEditorViewportToolbarBase.h"
+#include "WorldPartition/WorldPartition.h"
 
 #define LOCTEXT_NAMESPACE "LevelViewportToolBar"
 
@@ -1269,6 +1272,12 @@ void SLevelViewportToolBar::FillShowMenu(UToolMenu* Menu) const
 					FNewToolMenuDelegate::CreateStatic(&SLevelViewportToolBar::FillShowFoliageTypesMenu, Viewport), false, FSlateIcon(FAppStyle::Get().GetStyleSetName(), "ShowFlagsMenu.SubMenu.FoliageTypes"));
 			}
 
+			// Show 'HLODs' sub-menu is dynamically generated when the user enters 'show' menu
+			if (World->IsPartitionedWorld())
+			{
+				Section.AddSubMenu("ShowHLODsMenu", LOCTEXT("ShowHLODsMenu", "HLODs"), LOCTEXT("ShowHLODsMenu_ToolTip", "Settings for HLODs in editor"),
+					FNewToolMenuDelegate::CreateSP(this, &SLevelViewportToolBar::FillShowHLODsMenu), false, FSlateIcon(FAppStyle::Get().GetStyleSetName(), "ShowFlagsMenu.SubMenu.HLODs"));
+			}
 		}
 	}
 }
@@ -1493,6 +1502,131 @@ void SLevelViewportToolBar::FillShowFoliageTypesMenu(UToolMenu* Menu, TWeakPtr<c
 	}
 }
 
+double SLevelViewportToolBar::OnGetHLODInEditorMinDrawDistanceValue() const
+{
+	IWorldPartitionEditorModule* WorldPartitionEditorModule = FModuleManager::GetModulePtr<IWorldPartitionEditorModule>("WorldPartitionEditor");
+	return WorldPartitionEditorModule ? WorldPartitionEditorModule->GetHLODInEditorMinDrawDistance() : 0;
+}
+
+void SLevelViewportToolBar::OnHLODInEditorMinDrawDistanceValueChanged(double NewValue) const
+{
+	IWorldPartitionEditorModule* WorldPartitionEditorModule = FModuleManager::GetModulePtr<IWorldPartitionEditorModule>("WorldPartitionEditor");
+	if (WorldPartitionEditorModule)
+	{
+		WorldPartitionEditorModule->SetHLODInEditorMinDrawDistance(NewValue);
+		GEditor->RedrawLevelEditingViewports(true);
+	}
+}
+
+double SLevelViewportToolBar::OnGetHLODInEditorMaxDrawDistanceValue() const
+{
+	IWorldPartitionEditorModule* WorldPartitionEditorModule = FModuleManager::GetModulePtr<IWorldPartitionEditorModule>("WorldPartitionEditor");
+	return WorldPartitionEditorModule ? WorldPartitionEditorModule->GetHLODInEditorMaxDrawDistance() : 0;
+}
+
+void SLevelViewportToolBar::OnHLODInEditorMaxDrawDistanceValueChanged(double NewValue) const
+{
+	IWorldPartitionEditorModule* WorldPartitionEditorModule = FModuleManager::GetModulePtr<IWorldPartitionEditorModule>("WorldPartitionEditor");
+	if (WorldPartitionEditorModule)
+	{
+		WorldPartitionEditorModule->SetHLODInEditorMaxDrawDistance(NewValue);
+		GEditor->RedrawLevelEditingViewports(true);
+	}
+}
+
+void SLevelViewportToolBar::FillShowHLODsMenu(UToolMenu* Menu) const
+{
+	auto ViewportPtr = Viewport.Pin();
+	if (!ViewportPtr.IsValid())
+	{
+		return;
+	}
+
+	UWorld* World = ViewportPtr->GetWorld();
+	UWorldPartition* WorldPartition = World ? World->GetWorldPartition() : nullptr;
+	IWorldPartitionEditorModule* WorldPartitionEditorModule = FModuleManager::GetModulePtr<IWorldPartitionEditorModule>("WorldPartitionEditor");
+	
+	if (WorldPartition == nullptr || WorldPartitionEditorModule == nullptr)
+	{
+		return;
+	}
+
+	FText HLODInEditorDisallowedReason;
+	const bool bHLODInEditorAllowed = WorldPartitionEditorModule->IsHLODInEditorAllowed(World, &HLODInEditorDisallowedReason);
+
+	//Menu->AddMenuEntry(NAME_None, FToolMenuEntry::InitMenuEntry("HLODsInEditorUnavailable", HLODInEditorDisallowedReason, HLODInEditorDisallowedReason, FSlateIcon(), FToolUIAction()));
+
+	// Show HLODs
+	{
+		FToolUIAction UIAction;
+		UIAction.ExecuteAction = FToolMenuExecuteAction::CreateLambda([WorldPartitionEditorModule](const FToolMenuContext& InContext) { WorldPartitionEditorModule->SetShowHLODsInEditor(!WorldPartitionEditorModule->GetShowHLODsInEditor()); });
+		UIAction.CanExecuteAction = FToolMenuCanExecuteAction::CreateLambda([bHLODInEditorAllowed](const FToolMenuContext& InContext) { return bHLODInEditorAllowed; });
+		UIAction.GetActionCheckState = FToolMenuGetActionCheckState::CreateLambda([WorldPartitionEditorModule](const FToolMenuContext& InContext) { return WorldPartitionEditorModule->GetShowHLODsInEditor() ? ECheckBoxState::Checked : ECheckBoxState::Unchecked; });
+		FToolMenuEntry MenuEntry = FToolMenuEntry::InitMenuEntry("ShowHLODs", LOCTEXT("ShowHLODs", "Show HLODs"), bHLODInEditorAllowed ? LOCTEXT("ShowHLODsToolTip", "Show/Hide HLODs") : HLODInEditorDisallowedReason, FSlateIcon(), UIAction, EUserInterfaceActionType::ToggleButton);
+		Menu->AddMenuEntry(NAME_None, MenuEntry);
+	}
+
+	// Show HLODs Over Loaded Regions
+	{
+		FToolUIAction UIAction;
+		UIAction.ExecuteAction = FToolMenuExecuteAction::CreateLambda([WorldPartitionEditorModule](const FToolMenuContext& InContext) { WorldPartitionEditorModule->SetShowHLODsOverLoadedRegions(!WorldPartitionEditorModule->GetShowHLODsOverLoadedRegions()); });
+		UIAction.CanExecuteAction = FToolMenuCanExecuteAction::CreateLambda([bHLODInEditorAllowed](const FToolMenuContext& InContext) { return bHLODInEditorAllowed; });
+		UIAction.GetActionCheckState = FToolMenuGetActionCheckState::CreateLambda([WorldPartitionEditorModule](const FToolMenuContext& InContext) { return WorldPartitionEditorModule->GetShowHLODsOverLoadedRegions() ? ECheckBoxState::Checked : ECheckBoxState::Unchecked; });
+		FToolMenuEntry ShowHLODsEntry = FToolMenuEntry::InitMenuEntry("ShowHLODsOverLoadedRegions", LOCTEXT("ShowHLODsOverLoadedRegions", "Show HLODs Over Loaded Regions"), bHLODInEditorAllowed ? LOCTEXT("ShowHLODsOverLoadedRegions_ToolTip", "Show/Hide HLODs over loaded actors or regions") : HLODInEditorDisallowedReason, FSlateIcon(), UIAction, EUserInterfaceActionType::ToggleButton);
+		Menu->AddMenuEntry(NAME_None, ShowHLODsEntry);
+	}
+
+	// Min/Max Draw Distance
+	{
+		const double MinDrawDistanceMinValue = 0;
+		const double MinDrawDistanceMaxValue = 102400;
+
+		const double MaxDrawDistanceMinValue = 0;
+		const double MaxDrawDistanceMaxValue = 1638400;
+
+		TSharedRef<SSpinBox<double>> MinDrawDistanceSpinBox = SNew(SSpinBox<double>)
+			.MinValue(MinDrawDistanceMinValue)
+			.MaxValue(MinDrawDistanceMaxValue)
+			.IsEnabled(bHLODInEditorAllowed)
+			.Value(this, &SLevelViewportToolBar::OnGetHLODInEditorMinDrawDistanceValue)
+			.OnValueChanged(this, &SLevelViewportToolBar::OnHLODInEditorMinDrawDistanceValueChanged)
+			.ToolTipText(bHLODInEditorAllowed ? LOCTEXT("HLODsInEditor_MinDrawDistance_Tooltip", "Sets the minimum distance at which HLOD will be rendered") : HLODInEditorDisallowedReason);
+
+		TSharedRef<SSpinBox<double>> MaxDrawDistanceSpinBox = SNew(SSpinBox<double>)
+			.MinValue(MaxDrawDistanceMinValue)
+			.MaxValue(MaxDrawDistanceMaxValue)
+			.IsEnabled(bHLODInEditorAllowed)
+			.Value(this, &SLevelViewportToolBar::OnGetHLODInEditorMaxDrawDistanceValue)
+			.OnValueChanged(this, &SLevelViewportToolBar::OnHLODInEditorMaxDrawDistanceValueChanged)
+			.ToolTipText(bHLODInEditorAllowed ? LOCTEXT("HLODsInEditor_MaxDrawDistance_Tooltip", "Sets the maximum distance at which HLODs will be rendered") : HLODInEditorDisallowedReason);
+
+		auto CreateDrawDistanceWidget = [](TSharedRef<SSpinBox<double>> InSpinBoxWidget)
+		{
+			return SNew(SBox)
+				.HAlign(HAlign_Right)
+				[
+					SNew(SBox)
+						.Padding(FMargin(0.0f, 0.0f, 0.0f, 0.0f))
+						.WidthOverride(100.0f)
+						[
+							SNew(SBorder)
+								.BorderImage(FAppStyle::Get().GetBrush("Menu.WidgetBorder"))
+								.Padding(FMargin(1.0f))
+								[
+									InSpinBoxWidget
+								]
+						]
+				];
+		};
+			
+		FToolMenuEntry MinDrawDistanceMenuEntry = FToolMenuEntry::InitWidget("Min Draw Distance", CreateDrawDistanceWidget(MinDrawDistanceSpinBox), LOCTEXT("MinDrawDistance", "Min Draw Distance"));
+		Menu->AddMenuEntry(NAME_None, MinDrawDistanceMenuEntry);
+
+		FToolMenuEntry MaxDrawDistanceMenuEntry = FToolMenuEntry::InitWidget("Max Draw Distance", CreateDrawDistanceWidget(MaxDrawDistanceSpinBox), LOCTEXT("MaxDrawDistance", "Max Draw Distance"));
+		Menu->AddMenuEntry(NAME_None, MaxDrawDistanceMenuEntry);
+	}
+}
+
 TWeakObjectPtr<UWorld> SLevelViewportToolBar::GetWorld() const
 {
 	if (Viewport.IsValid())
diff --git a/Engine/Source/Editor/LevelEditor/Private/SLevelViewportToolBar.h b/Engine/Source/Editor/LevelEditor/Private/SLevelViewportToolBar.h
--- a/Engine/Source/Editor/LevelEditor/Private/SLevelViewportToolBar.h
+++ b/Engine/Source/Editor/LevelEditor/Private/SLevelViewportToolBar.h
@@ -235,6 +235,11 @@ private:
 	EVisibility GetScalabilityWarningVisibility() const;
 	TSharedRef<SWidget> GetScalabilityWarningMenuContent() const;
 
+	double OnGetHLODInEditorMinDrawDistanceValue() const;
+	void OnHLODInEditorMinDrawDistanceValueChanged(double NewValue) const;
+	double OnGetHLODInEditorMaxDrawDistanceValue() const;
+	void OnHLODInEditorMaxDrawDistanceValueChanged(double NewValue) const;
+
 private:
 	/**
 	 * Generates the toolbar show layers menu content 
@@ -251,6 +256,14 @@ private:
 	 */
 	static void FillShowFoliageTypesMenu(UToolMenu* Menu, TWeakPtr<class SLevelViewport> Viewport);
 
+	/**
+	 * Generates 'Show HLODs' menu content for a viewport
+	 *
+	 * @param Menu		The tool menu
+	 * @param Viewport	Target vieport
+	 */
+	void FillShowHLODsMenu(UToolMenu* Menu) const;
+
 	/** Generates the layout sub-menu content */
 	void GenerateViewportConfigsMenu(UToolMenu* Menu) const;
 
diff --git a/Engine/Source/Editor/WorldPartitionEditor/Private/WorldPartition/HLOD/HLODEditorData.cpp b/Engine/Source/Editor/WorldPartitionEditor/Private/WorldPartition/HLOD/HLODEditorData.cpp
--- a/Engine/Source/Editor/WorldPartitionEditor/Private/WorldPartition/HLOD/HLODEditorData.cpp
+++ b/Engine/Source/Editor/WorldPartitionEditor/Private/WorldPartition/HLOD/HLODEditorData.cpp
@@ -105,6 +105,12 @@ static void GatherLoadedActorsBounds(TArray<FBoundsWithVolume>& OutLoadedBounds,
 	}
 }
 
+void FWorldPartitionHLODEditorData::ClearLoadedActorsState()
+{
+	// Increment update counter - used to quickly find out if an HLOD actor needs be hidden without having to flag the whole hierarchy
+	LastStateUpdate++;
+}
+
 void FWorldPartitionHLODEditorData::UpdateLoadedActorsState()
 {
 	TRACE_CPUPROFILER_EVENT_SCOPE(FWorldPartitionHLODEditorData::UpdateLoadedActorsState);
@@ -189,7 +195,7 @@ void FWorldPartitionHLODEditorData::UpdateLoadedActorsState()
 	}
 }
 
-void FWorldPartitionHLODEditorData::UpdateVisibility(const FVector& InCameraLocation, bool bInForceVisibilityUpdate)
+void FWorldPartitionHLODEditorData::UpdateVisibility(const FVector& InCameraLocation, double InMinDrawDistance, double InMaxDrawDistance, bool bInForceVisibilityUpdate)
 {
 	TRACE_CPUPROFILER_EVENT_SCOPE(FWorldPartitionHLODEditorData::UpdateVisibility);
 
@@ -208,11 +214,11 @@ void FWorldPartitionHLODEditorData::UpdateVisibility(const FVector& InCameraLoca
 	for (const auto HLODActorNode : TopLevelHLODActorNodes)
 	{
 		// Recurse from the top level HLOD down to HLOD0
-		HLODActorNode->UpdateVisibility(InCameraLocation, /*bInForceHidden*/false, bInForceVisibilityUpdate, LastStateUpdate);
+		HLODActorNode->UpdateVisibility(InCameraLocation, InMinDrawDistance, InMaxDrawDistance, /*bInForceHidden*/false, bInForceVisibilityUpdate, LastStateUpdate);
 	}
 }
 
-void FHLODSceneNode::UpdateVisibility(const FVector& InCameraLocation, bool bInForceHidden, bool bInForceVisibilityUpdate, int32 InLastStateUpdate)
+void FHLODSceneNode::UpdateVisibility(const FVector& InCameraLocation, double InMinDrawDistance, double InMaxDrawDistance, bool bInForceHidden, bool bInForceVisibilityUpdate, int32 InLastStateUpdate)
 {
 	bool bNodeShouldBeVisible = !bInForceHidden;
 
@@ -232,10 +238,19 @@ void FHLODSceneNode::UpdateVisibility(const FVector& InCameraLocation, bool bInF
 			const bool bHasLoadedChildrenHLODs = Algo::AnyOf(ChildrenHLODs, [](FHLODSceneNode* ChildHLODNode) { return ChildHLODNode->HLODActorHandle.IsLoaded(); });
 			if (bHasLoadedChildrenHLODs)
 			{
-				const float DistanceSquared = Bounds.ComputeSquaredDistanceFromBoxToPoint(InCameraLocation);
+				const double DistanceSquared = Bounds.ComputeSquaredDistanceFromBoxToPoint(InCameraLocation);
 				const bool bNearCulled = DistanceSquared < FMath::Square(HLODActor->GetMinVisibleDistance());
 				bNodeShouldBeVisible = !bNearCulled;
 			}
+			
+			// Apply HLOD min/max draw distance user setting
+			if (bNodeShouldBeVisible && (InMinDrawDistance != 0 || InMaxDrawDistance != 0))
+			{
+				const double DistanceSquared = Bounds.ComputeSquaredDistanceFromBoxToPoint(InCameraLocation);
+				const bool bNearCulled = InMinDrawDistance == 0 ? false : DistanceSquared < FMath::Square(InMinDrawDistance);
+				const bool bFarCulled = InMaxDrawDistance == 0 ? false : DistanceSquared > FMath::Square(InMaxDrawDistance);
+				bNodeShouldBeVisible = !bNearCulled && !bFarCulled;
+			}
 		}
 
 		if (bInForceVisibilityUpdate || bCachedIsVisible != bNodeShouldBeVisible)
@@ -248,7 +263,7 @@ void FHLODSceneNode::UpdateVisibility(const FVector& InCameraLocation, bool bInF
 	const bool bForceHideChildren = bNodeShouldBeVisible || bInForceHidden;
 	for (FHLODSceneNode* Child : ChildrenHLODs)
 	{
-		Child->UpdateVisibility(InCameraLocation, bForceHideChildren, bInForceVisibilityUpdate, InLastStateUpdate);
+		Child->UpdateVisibility(InCameraLocation, InMinDrawDistance, InMaxDrawDistance, bForceHideChildren, bInForceVisibilityUpdate, InLastStateUpdate);
 	}
 }
 
diff --git a/Engine/Source/Editor/WorldPartitionEditor/Private/WorldPartition/HLOD/HLODEditorSubsystem.cpp b/Engine/Source/Editor/WorldPartitionEditor/Private/WorldPartition/HLOD/HLODEditorSubsystem.cpp
--- a/Engine/Source/Editor/WorldPartitionEditor/Private/WorldPartition/HLOD/HLODEditorSubsystem.cpp
+++ b/Engine/Source/Editor/WorldPartitionEditor/Private/WorldPartition/HLOD/HLODEditorSubsystem.cpp
@@ -5,16 +5,17 @@
 #include "Editor.h"
 #include "Engine/Engine.h"
 #include "Engine/World.h"
+#include "GameFramework/WorldSettings.h"
 #include "Subsystems/UnrealEditorSubsystem.h"
+#include "WorldPartitionEditorModule.h"
 #include "WorldPartition/HLOD/HLODEditorData.h"
 #include "WorldPartition/WorldPartition.h"
 #include "WorldPartition/WorldPartitionSubsystem.h"
 
-
 static TAutoConsoleVariable<bool> CVarHLODInEditorEnabled(
-	TEXT("wp.Editor.HLOD"),
+	TEXT("wp.Editor.HLOD.AllowShowingHLODsInEditor"),
 	false,
-	TEXT("Show World Partition HLODs in the editor."));
+	TEXT("Allow showing World Partition HLODs in the editor."));
 
 
 UWorldPartitionHLODEditorSubsystem::UWorldPartitionHLODEditorSubsystem()
@@ -27,7 +28,10 @@ UWorldPartitionHLODEditorSubsystem::~UWorldPartitionHLODEditorSubsystem()
 
 bool UWorldPartitionHLODEditorSubsystem::IsHLODInEditorEnabled()
 {
-	return CVarHLODInEditorEnabled.GetValueOnGameThread();
+	const IWorldPartitionEditorModule* WorldPartitionEditorModule = FModuleManager::GetModulePtr<IWorldPartitionEditorModule>("WorldPartitionEditor");
+	const bool bShowHLODsInEditorForWorld = WorldPartitionEditorModule && WorldPartitionEditorModule->IsHLODInEditorAllowed(GetWorld());
+	const bool bShowHLODsInEditorUserSetting = WorldPartitionEditorModule && WorldPartitionEditorModule->GetShowHLODsInEditor();
+	return CVarHLODInEditorEnabled.GetValueOnGameThread() && bShowHLODsInEditorUserSetting && bShowHLODsInEditorForWorld;
 }
 
 bool UWorldPartitionHLODEditorSubsystem::DoesSupportWorldType(const EWorldType::Type WorldType) const
@@ -43,6 +47,10 @@ void UWorldPartitionHLODEditorSubsystem::Initialize(FSubsystemCollectionBase& Co
 	Super::Initialize(Collection);
 
 	bForceHLODStateUpdate = true;
+	CachedCameraLocation = FVector::Zero();
+	CachedHLODMinDrawDistance = 0;
+	CachedHLODMaxDrawDistance = 0;
+	bCachedShowHLODsOverLoadedRegions = false;
 	
 	GetWorld()->OnWorldPartitionInitialized().AddUObject(this, &UWorldPartitionHLODEditorSubsystem::OnWorldPartitionInitialized);
 	GetWorld()->OnWorldPartitionUninitialized().AddUObject(this, &UWorldPartitionHLODEditorSubsystem::OnWorldPartitionUninitialized);
@@ -108,9 +116,41 @@ void UWorldPartitionHLODEditorSubsystem::Tick(float DeltaTime)
 		
 		if (IsHLODInEditorEnabled())
 		{
-			if (bForceHLODStateUpdate)
+			bool bForceHLODVisibilityUpdate = false;
+
+			IWorldPartitionEditorModule* WorldPartitionEditorModule = FModuleManager::GetModulePtr<IWorldPartitionEditorModule>("WorldPartitionEditor");
+
+			// "Show HLODs over loaded region" option changed ?
+			if (WorldPartitionEditorModule->GetShowHLODsOverLoadedRegions() != bCachedShowHLODsOverLoadedRegions)
+			{ 
+				bCachedShowHLODsOverLoadedRegions = WorldPartitionEditorModule->GetShowHLODsOverLoadedRegions();
+				bForceHLODVisibilityUpdate = true;
+
+				// Clear loading state of actors if we are going to always show HLODs
+				if (bCachedShowHLODsOverLoadedRegions)
+				{
+					HLODEditorData->ClearLoadedActorsState();
+				}
+				else
+				{
+					bForceHLODStateUpdate = true;
+				}
+			}
+
+			// Actors or regions were loaded ?
+			if (bForceHLODStateUpdate && !bCachedShowHLODsOverLoadedRegions)
 			{
 				HLODEditorData->UpdateLoadedActorsState();
+				bForceHLODVisibilityUpdate = true;
+			}
+			
+			// Min/Max draw distance for HLODs was changed ?
+			if (WorldPartitionEditorModule->GetHLODInEditorMinDrawDistance() != CachedHLODMinDrawDistance ||
+				WorldPartitionEditorModule->GetHLODInEditorMaxDrawDistance() != CachedHLODMaxDrawDistance)
+			{
+				CachedHLODMinDrawDistance = WorldPartitionEditorModule->GetHLODInEditorMinDrawDistance();
+				CachedHLODMaxDrawDistance = WorldPartitionEditorModule->GetHLODInEditorMaxDrawDistance();
+				bForceHLODVisibilityUpdate = true;
 			}
 
 			UUnrealEditorSubsystem* UnrealEditorSubsystem = GEditor->GetEditorSubsystem<UUnrealEditorSubsystem>();
@@ -120,14 +160,18 @@ void UWorldPartitionHLODEditorSubsystem::Tick(float DeltaTime)
 				FRotator CameraRotation;
 				UnrealEditorSubsystem->GetLevelViewportCameraInfo(CameraLocation, CameraRotation);
 
-				if (bForceHLODStateUpdate || CameraLocation != CachedCameraLocation)
+				// Camera was moved ?
+				const bool bCameraMoved = CameraLocation != CachedCameraLocation;
+				if (bCameraMoved)
 				{
 					CachedCameraLocation = CameraLocation;
+				}
 
-					HLODEditorData->UpdateVisibility(CameraLocation, bForceHLODStateUpdate);
+				if (bForceHLODVisibilityUpdate || bCameraMoved)
+				{					
+					HLODEditorData->UpdateVisibility(CameraLocation, CachedHLODMinDrawDistance, CachedHLODMaxDrawDistance, bForceHLODVisibilityUpdate);
 				}
 			}
-
 			bForceHLODStateUpdate = false;
 		}
 	}
diff --git a/Engine/Source/Editor/WorldPartitionEditor/Private/WorldPartition/WorldPartitionEditorSettings.cpp b/Engine/Source/Editor/WorldPartitionEditor/Private/WorldPartition/WorldPartitionEditorSettings.cpp
--- a/Engine/Source/Editor/WorldPartitionEditor/Private/WorldPartition/WorldPartitionEditorSettings.cpp
+++ b/Engine/Source/Editor/WorldPartitionEditor/Private/WorldPartition/WorldPartitionEditorSettings.cpp
@@ -3,6 +3,22 @@
 #include "WorldPartition/WorldPartitionEditorSettings.h"
 #include "WorldPartition/WorldPartition.h"
 
+UWorldPartitionEditorSettings::UWorldPartitionEditorSettings()
+{
+	CommandletClass = UWorldPartitionConvertCommandlet::StaticClass();
+	InstancedFoliageGridSize = 25600;
+	MinimapLowQualityWorldUnitsPerPixelThreshold = 12800;
+	bEnableLoadingInEditor = true;
+	bEnableStreamingGenerationLogOnPIE = true;
+	bShowHLODsInEditor = true;
+	bShowHLODsOverLoadedRegions = false;
+	HLODMinDrawDistance = 12800;
+	HLODMaxDrawDistance = 0;
+	bDisablePIE = false;
+	bDisableBugIt = false;
+	bAdvancedMode = true;
+}
+
 void UWorldPartitionEditorSettings::PostEditChangeProperty(struct FPropertyChangedEvent& PropertyChangedEvent)
 {
 	Super::PostEditChangeProperty(PropertyChangedEvent);
diff --git a/Engine/Source/Editor/WorldPartitionEditor/Private/WorldPartitionEditorModule.cpp b/Engine/Source/Editor/WorldPartitionEditor/Private/WorldPartitionEditorModule.cpp
--- a/Engine/Source/Editor/WorldPartitionEditor/Private/WorldPartitionEditorModule.cpp
+++ b/Engine/Source/Editor/WorldPartitionEditor/Private/WorldPartitionEditorModule.cpp
@@ -7,6 +7,7 @@
 #include "WorldPartition/WorldPartitionMiniMapBuilder.h"
 #include "WorldPartition/WorldPartitionLandscapeSplineMeshesBuilder.h"
 #include "WorldPartition/WorldPartitionActorLoaderInterface.h"
+#include "WorldPartition/LoaderAdapter/LoaderAdapterShape.h"
 #include "WorldPartition/SWorldPartitionEditor.h"
 #include "WorldPartition/SWorldPartitionEditorGridSpatialHash.h"
 #include "WorldPartition/Customizations/WorldPartitionDetailsCustomization.h"
@@ -16,10 +17,12 @@
 #include "WorldPartition/SWorldPartitionConvertDialog.h"
 #include "WorldPartition/WorldPartitionConvertOptions.h"
 #include "WorldPartition/WorldPartitionEditorSettings.h"
+#include "WorldPartition/HLOD/HLODActor.h"
 #include "WorldPartition/HLOD/SWorldPartitionBuildHLODsDialog.h"
 #include "WorldPartition/WorldPartitionClassDescRegistry.h"
 
 #include "LevelEditor.h"
+#include "LevelEditorViewport.h"
 #include "Framework/Application/SlateApplication.h"
 #include "Framework/MultiBox/MultiBoxBuilder.h"
 
@@ -50,6 +53,7 @@
 #include "WorkspaceMenuStructureModule.h"
 #include "Framework/Docking/LayoutExtender.h"
 #include "Widgets/Docking/SDockTab.h"
+#include "Filters/CustomClassFilterData.h"
 
 #include "Styling/AppStyle.h"
 #include "WorldPartition/ContentBundle/SContentBundleBrowser.h"
@@ -104,37 +108,61 @@ static bool CanLoadUnloadSelectedVolumes(TArray<TWeakObjectPtr<AActor>> Volumes,
 	return false;
 }
 
-static void CreateLevelViewportContextMenuEntries(FMenuBuilder& MenuBuilder, TArray<TWeakObjectPtr<AActor>> Volumes)
+static void CreateLevelViewportContextMenuEntries(FMenuBuilder& MenuBuilder, TArray<TWeakObjectPtr<AActor>> Volumes, FBox SelectionBox)
 {
 	MenuBuilder.BeginSection("WorldPartition", LOCTEXT("WorldPartition", "World Partition"));
 	
-	MenuBuilder.AddMenuEntry(
-		LOCTEXT("WorldPartitionLoad", "Load selected volumes"),
-		LOCTEXT("WorldPartitionLoad_Tooltip", "Load selected volumes"),
-		FSlateIcon(),
-		FUIAction(
-			FExecuteAction::CreateStatic(OnSelectedWorldPartitionVolumesToggleLoading, Volumes, true),
-			FCanExecuteAction::CreateLambda([Volumes]
-			{
-				return CanLoadUnloadSelectedVolumes(Volumes, true);
-			})
-		),
-		NAME_None,
-		EUserInterfaceActionType::Button);
-
-	MenuBuilder.AddMenuEntry(
-		LOCTEXT("WorldPartitionUnload", "Unload selected volumes"),
-		LOCTEXT("WorldPartitionUnload_Tooltip", "Load selected volumes"),
-		FSlateIcon(),
-		FUIAction(
-			FExecuteAction::CreateStatic(OnSelectedWorldPartitionVolumesToggleLoading, Volumes, false),
-			FCanExecuteAction::CreateLambda([Volumes]
+	if (!Volumes.IsEmpty())
+	{
+		MenuBuilder.AddMenuEntry(
+			LOCTEXT("WorldPartitionLoad", "Load selected volumes"),
+			LOCTEXT("WorldPartitionLoad_Tooltip", "Load selected volumes"),
+			FSlateIcon(),
+			FUIAction(
+				FExecuteAction::CreateStatic(OnSelectedWorldPartitionVolumesToggleLoading, Volumes, true),
+				FCanExecuteAction::CreateLambda([Volumes]
+				{
+					return CanLoadUnloadSelectedVolumes(Volumes, true);
+				})
+			),
+			NAME_None,
+			EUserInterfaceActionType::Button);
+
+		MenuBuilder.AddMenuEntry(
+			LOCTEXT("WorldPartitionUnload", "Unload selected volumes"),
+			LOCTEXT("WorldPartitionUnload_Tooltip", "Load selected volumes"),
+			FSlateIcon(),
+			FUIAction(
+				FExecuteAction::CreateStatic(OnSelectedWorldPartitionVolumesToggleLoading, Volumes, false),
+				FCanExecuteAction::CreateLambda([Volumes]
 				{
 					return CanLoadUnloadSelectedVolumes(Volumes, false);
 				})
-		),
-		NAME_None,
-		EUserInterfaceActionType::Button);
+			),
+			NAME_None,
+			EUserInterfaceActionType::Button);
+	}
+
+	// Load Region From Selection
+	if (GCurrentLevelEditingViewportClient && SelectionBox.GetSize().Size2D() > 0)
+	{
+		TWeakObjectPtr<UWorld> World = GCurrentLevelEditingViewportClient->GetWorld();
+		TWeakObjectPtr<UWorldPartition> WorldPartition = World->GetWorldPartition();
+
+		FUIAction LoadRegion(
+			FExecuteAction::CreateLambda([World, WorldPartition, SelectionBox]()
+			{
+				if (World.IsValid() && WorldPartition.IsValid())
+				{
+					UWorldPartitionEditorLoaderAdapter* EditorLoaderAdapter = WorldPartition.Get()->CreateEditorLoaderAdapter<FLoaderAdapterShape>(World.Get(), SelectionBox, TEXT("Loaded Region"));
+					EditorLoaderAdapter->GetLoaderAdapter()->SetUserCreated(true);
+					EditorLoaderAdapter->GetLoaderAdapter()->Load();
+				}
+			})
+		);
+
+		MenuBuilder.AddMenuEntry(LOCTEXT("LoadRegionFromSelection", "Load Region From Selection"), LOCTEXT("LoadRegionFromSelection_Tooltip", "Load region from selected actor(s) bounds"), FSlateIcon(), LoadRegion);
+	}
 	
 	MenuBuilder.EndSection();
 }
@@ -143,22 +171,28 @@ static TSharedRef<FExtender> OnExtendLevelEditorMenu(const TSharedRef<FUICommand
 {
 	TSharedRef<FExtender> Extender(new FExtender());
 
-	TArray<TWeakObjectPtr<AActor> > Volumes;
+	TArray<TWeakObjectPtr<AActor>> Volumes;
+	FBoxSphereBounds::Builder BoundsBuilder;
+
 	for (AActor* Actor : SelectedActors)
 	{
 		if (Actor->Implements<UWorldPartitionActorLoaderInterface>())
 		{
 			Volumes.Add(Actor);
 		}
+
+		FBoxSphereBounds ActorBounds;
+		Actor->GetActorBounds(false, ActorBounds.Origin, ActorBounds.BoxExtent);
+		BoundsBuilder += ActorBounds;
 	}
 
-	if (Volumes.Num())
+	if (!Volumes.IsEmpty() || BoundsBuilder.IsValid())
 	{
 		Extender->AddMenuExtension(
 			"ActorTypeTools",
 			EExtensionHook::After,
 			nullptr,
-			FMenuExtensionDelegate::CreateStatic(&CreateLevelViewportContextMenuEntries, Volumes));
+			FMenuExtensionDelegate::CreateStatic(&CreateLevelViewportContextMenuEntries, Volumes, FBoxSphereBounds(BoundsBuilder).GetBox()));
 	}
 
 	return Extender;
@@ -180,6 +214,34 @@ void FWorldPartitionEditorModule::StartupModule()
 	FWorldPartitionClassDescRegistry().Get().Initialize();
 
 	CleanupExternalObjectsEmptyFolders();
+
+	EditorInitializedHandle = FEditorDelegates::OnEditorInitialized.AddLambda([this](double TimeToInitializeEditor)
+	{
+		// Register the Scene Outliner "World" filter category
+		if (FModuleManager::Get().IsModuleLoaded("LevelEditor"))
+		{
+			FLevelEditorModule& LevelEditorModule = FModuleManager::GetModuleChecked<FLevelEditorModule>("LevelEditor");
+			TSharedPtr<FFilterCategory> CommonFilterCategory = LevelEditorModule.GetOutlinerFilterCategory(FLevelEditorOutlinerBuiltInCategories::Common());
+			TSharedPtr<FFilterCategory> WorldFilterCategory = MakeShared<FFilterCategory>(LOCTEXT("WorldFilterCategory", "World"), FText::GetEmpty());
+
+			TArray<UClass*> WorldActorClasses =
+			{
+				AWorldPartitionHLOD::StaticClass()
+			};
+
+			for (UClass* Class : WorldActorClasses)
+			{
+				TSharedRef<FCustomClassFilterData> ClassFilterData = MakeShared<FCustomClassFilterData>(AWorldPartitionHLOD::StaticClass(), WorldFilterCategory, FLinearColor::White);
+
+				if (CommonFilterCategory.IsValid())
+				{
+					ClassFilterData->AddCategory(CommonFilterCategory);
+				}
+
+				LevelEditorModule.AddCustomClassFilterToOutliner(ClassFilterData);
+			}
+		}
+	});
 }
 
 void FWorldPartitionEditorModule::ShutdownModule()
@@ -214,6 +276,8 @@ void FWorldPartitionEditorModule::ShutdownModule()
 	}
 
 	WaitForCleanupExternalObjectsEmptyFolders();
+
+	FEditorDelegates::OnEditorInitialized.Remove(EditorInitializedHandle);
 }
 
 void FWorldPartitionEditorModule::RegisterMenus()
@@ -331,6 +395,83 @@ void FWorldPartitionEditorModule::SetAdvancedMode(bool bInAdvancedMode)
 	GetMutableDefault<UWorldPartitionEditorSettings>()->bAdvancedMode = bInAdvancedMode;
 }
 
+bool FWorldPartitionEditorModule::GetShowHLODsInEditor() const
+{
+	return GetDefault<UWorldPartitionEditorSettings>()->bShowHLODsInEditor;
+}
+
+void FWorldPartitionEditorModule::SetShowHLODsInEditor(bool bInShowHLODsInEditor)
+{
+	GetMutableDefault<UWorldPartitionEditorSettings>()->bShowHLODsInEditor = bInShowHLODsInEditor;
+}
+
+bool FWorldPartitionEditorModule::GetShowHLODsOverLoadedRegions() const
+{
+	return GetDefault<UWorldPartitionEditorSettings>()->bShowHLODsOverLoadedRegions;
+}
+
+void FWorldPartitionEditorModule::SetShowHLODsOverLoadedRegions(bool bInShowHLODsOverLoadedRegions)
+{
+	GetMutableDefault<UWorldPartitionEditorSettings>()->bShowHLODsOverLoadedRegions = bInShowHLODsOverLoadedRegions;
+}
+
+double FWorldPartitionEditorModule::GetHLODInEditorMinDrawDistance() const
+{
+	return GetDefault<UWorldPartitionEditorSettings>()->HLODMinDrawDistance;
+}
+
+void FWorldPartitionEditorModule::SetHLODInEditorMinDrawDistance(double InMinDrawDistance)
+{
+	GetMutableDefault<UWorldPartitionEditorSettings>()->HLODMinDrawDistance = InMinDrawDistance;
+}
+
+double FWorldPartitionEditorModule::GetHLODInEditorMaxDrawDistance() const
+{
+	return GetDefault<UWorldPartitionEditorSettings>()->HLODMaxDrawDistance;
+}
+
+void FWorldPartitionEditorModule::SetHLODInEditorMaxDrawDistance(double InMaxDrawDistance)
+{
+	GetMutableDefault<UWorldPartitionEditorSettings>()->HLODMaxDrawDistance = InMaxDrawDistance;
+}
+
+bool FWorldPartitionEditorModule::IsHLODInEditorAllowed(UWorld* InWorld, FText* OutDisallowedReason) const
+{
+	auto SetDissallowedReason = [OutDisallowedReason](const FText& DisallowedReason)
+	{
+		if (OutDisallowedReason)
+		{
+			*OutDisallowedReason = DisallowedReason;
+		}
+	};
+
+	if (!InWorld)
+	{
+		SetDissallowedReason(LOCTEXT("HLODInEditor_InvalidWorld", "Invalid world"));
+		return false;
+	}
+
+	if (!InWorld->IsPartitionedWorld())
+	{
+		SetDissallowedReason(LOCTEXT("HLODInEditor_NoWorldPartition", "World is non partitioned"));
+		return false;
+	}
+
+	if (!InWorld->GetWorldPartition()->IsStreamingEnabledInEditor())
+	{
+		SetDissallowedReason(LOCTEXT("HLODInEditor_StreamingDisabled", "Streaming is disabled for this world"));
+		return false;
+	}
+
+	if (!InWorld->GetWorldPartition()->IsHLODsInEditorAllowed())
+	{
+		SetDissallowedReason(LOCTEXT("HLODInEditor_StreamingDisabled", "HLOD in editor is disabled for this world"));
+		return false;
+	}
+
+	return true;
+}
+
 void FWorldPartitionEditorModule::OnConvertMap()
 {
 	IContentBrowserSingleton& ContentBrowserSingleton = FModuleManager::LoadModuleChecked<FContentBrowserModule>("ContentBrowser").Get();
@@ -912,18 +1053,6 @@ void FWorldPartitionEditorModule::WaitForCleanupExternalObjectsEmptyFolders()
 	}
 }
 
-UWorldPartitionEditorSettings::UWorldPartitionEditorSettings()
-{
-	CommandletClass = UWorldPartitionConvertCommandlet::StaticClass();
-	InstancedFoliageGridSize = 25600;
-	MinimapLowQualityWorldUnitsPerPixelThreshold = 12800;
-	bEnableLoadingInEditor = true;
-	bEnableStreamingGenerationLogOnPIE = true;
-	bDisablePIE = false;
-	bDisableBugIt = false;
-	bAdvancedMode = true;
-}
-
 FString UWorldPartitionConvertOptions::ToCommandletArgs() const
 {
 	TStringBuilder<1024> CommandletArgsBuilder;
diff --git a/Engine/Source/Editor/WorldPartitionEditor/Public/WorldPartition/HLOD/HLODEditorData.h b/Engine/Source/Editor/WorldPartitionEditor/Public/WorldPartition/HLOD/HLODEditorData.h
--- a/Engine/Source/Editor/WorldPartitionEditor/Public/WorldPartition/HLOD/HLODEditorData.h
+++ b/Engine/Source/Editor/WorldPartitionEditor/Public/WorldPartition/HLOD/HLODEditorData.h
@@ -10,7 +10,7 @@ class UWorldPartition;
 // Represent an HLOD actor in the editor, loaded or not
 struct FHLODSceneNode
 {
-	void UpdateVisibility(const FVector& InCameraLocation, bool bInForceHidden, bool bInForceVisibilityUpdate, int32 InLastStateUpdate);
+	void UpdateVisibility(const FVector& InCameraLocation, double InMinDrawDistance, double InMaxDrawDistance, bool bInForceHidden, bool bInForceVisibilityUpdate, int32 InLastStateUpdate);
 
 	FHLODSceneNode* ParentHLOD = nullptr;
 	TArray<FHLODSceneNode*> ChildrenHLODs;
@@ -29,8 +29,9 @@ struct FWorldPartitionHLODEditorData
 public:
 	FWorldPartitionHLODEditorData(UWorldPartition* InWorldPartition);
 	
+	void ClearLoadedActorsState();
 	void UpdateLoadedActorsState();
-	void UpdateVisibility(const FVector& InCameraLocation, bool bForceVisibilityUpdate);
+	void UpdateVisibility(const FVector& InCameraLocation, double InMinDrawDistance, double InMaxDrawDistance, bool bForceVisibilityUpdate);
 
 	void SetHLODLoadingState(bool bInShouldBeLoaded);
 
diff --git a/Engine/Source/Editor/WorldPartitionEditor/Public/WorldPartition/HLOD/HLODEditorSubsystem.h b/Engine/Source/Editor/WorldPartitionEditor/Public/WorldPartition/HLOD/HLODEditorSubsystem.h
--- a/Engine/Source/Editor/WorldPartitionEditor/Public/WorldPartition/HLOD/HLODEditorSubsystem.h
+++ b/Engine/Source/Editor/WorldPartitionEditor/Public/WorldPartition/HLOD/HLODEditorSubsystem.h
@@ -42,10 +42,11 @@ public:
 	virtual void Tick(float DeltaTime) override;
 	virtual TStatId GetStatId() const override;
 	//~ End FTickableGameObject Interface
-
-	static bool IsHLODInEditorEnabled();
-
+	
 private:
+	bool IsHLODInEditorEnabled();
+	void SetHLODInEditorEnabled(bool bInEnable);
+
 	void OnWorldPartitionInitialized(UWorldPartition* InWorldPartition);
 	void OnWorldPartitionUninitialized(UWorldPartition* InWorldPartition);
 
@@ -55,6 +56,10 @@ private:
 	
 private:
 	FVector CachedCameraLocation;
+	double CachedHLODMinDrawDistance;
+	double CachedHLODMaxDrawDistance;
+	bool bCachedShowHLODsOverLoadedRegions;
+
 	bool bForceHLODStateUpdate;
 
 	TPimplPtr<FWorldPartitionHLODEditorData> HLODEditorData;
diff --git a/Engine/Source/Editor/WorldPartitionEditor/Public/WorldPartition/WorldPartitionEditorSettings.h b/Engine/Source/Editor/WorldPartitionEditor/Public/WorldPartition/WorldPartitionEditorSettings.h
--- a/Engine/Source/Editor/WorldPartitionEditor/Public/WorldPartition/WorldPartitionEditorSettings.h
+++ b/Engine/Source/Editor/WorldPartitionEditor/Public/WorldPartition/WorldPartitionEditorSettings.h
@@ -20,18 +20,30 @@ public:
 	UPROPERTY(Config, EditAnywhere, Category = MapConversion, Meta = (ToolTip = "Commandlet class to use for World Parition conversion"))
 	TSubclassOf<UWorldPartitionConvertCommandlet> CommandletClass;
 
-	UPROPERTY(Config, EditAnywhere, Category = Foliage, Meta = (ClampMin=3200, ToolTip= "Editor grid size used for instance foliage actors in World Partition worlds"))
+	UPROPERTY(Config, EditAnywhere, Category = Foliage, Meta = (ClampMin = 3200, ToolTip= "Editor grid size used for instance foliage actors in World Partition worlds"))
 	int32 InstancedFoliageGridSize;
 
 	UPROPERTY(Config, EditAnywhere, Category = MiniMap, Meta = (ClampMin = 100, ToolTip = "Threshold from which minimap generates a warning if its WorldUnitsPerPixel is above this value"))
 	int32 MinimapLowQualityWorldUnitsPerPixelThreshold;
 
-	UPROPERTY(Config, EditAnywhere, Category = WorldPartition, Meta = (ToolTip = "Wheter to enable dynamic loading in the editor through loading regions"))
+	UPROPERTY(Config, EditAnywhere, Category = WorldPartition, Meta = (ToolTip = "Whether to enable dynamic loading in the editor through loading regions"))
 	bool bEnableLoadingInEditor;
 
-	UPROPERTY(Config, EditAnywhere, Category = WorldPartition, Meta = (ToolTip = "Wheter to enable streaming generation log on PIE"))
+	UPROPERTY(Config, EditAnywhere, Category = WorldPartition, Meta = (ToolTip = "Whether to enable streaming generation log on PIE"))
 	bool bEnableStreamingGenerationLogOnPIE;
 
+	UPROPERTY(Config, EditAnywhere, Category = HLOD, Meta = (ToolTip = "Whether to show HLODs in the editor"))
+	bool bShowHLODsInEditor;
+
+	UPROPERTY(Config, EditAnywhere, Category = HLOD, Meta = (ToolTip = "Control display of HLODs in case actors are loaded"))
+	bool bShowHLODsOverLoadedRegions;
+
+	UPROPERTY(Config, EditAnywhere, Category = HLOD, Meta = (UIMin = 0, UIMax = 1638400, DisplayName = "HLOD Min Draw Distance", ToolTip = "Minimum distance at which HLODs should be displayed in editor"))
+	double HLODMinDrawDistance;
+
+	UPROPERTY(Config, EditAnywhere, Category = HLOD, Meta = (UIMin = 0, UIMax = 1638400, DisplayName = "HLOD Max Draw Distance", ToolTip = "Maximum distance at which HLODs should be displayed in editor"))
+	double HLODMaxDrawDistance;
+
 	bool bDisableBugIt;
 	bool bDisablePIE;
 	bool bAdvancedMode;
diff --git a/Engine/Source/Editor/WorldPartitionEditor/Public/WorldPartitionEditorModule.h b/Engine/Source/Editor/WorldPartitionEditor/Public/WorldPartitionEditorModule.h
--- a/Engine/Source/Editor/WorldPartitionEditor/Public/WorldPartitionEditorModule.h
+++ b/Engine/Source/Editor/WorldPartitionEditor/Public/WorldPartitionEditorModule.h
@@ -78,6 +78,20 @@ public:
 	virtual bool GetAdvancedMode() const override;
 	virtual void SetAdvancedMode(bool bInAdvancedMode) override;
 
+	virtual bool GetShowHLODsInEditor() const override;
+	virtual void SetShowHLODsInEditor(bool bInShowHLODsInEditor) override;
+ 
+	virtual bool GetShowHLODsOverLoadedRegions() const override;
+	virtual void SetShowHLODsOverLoadedRegions(bool bInShowHLODsOverLoadedRegions) override;
+
+	virtual double GetHLODInEditorMinDrawDistance() const override;
+	virtual void SetHLODInEditorMinDrawDistance(double InMinDrawDistance) override;
+
+	virtual double GetHLODInEditorMaxDrawDistance() const override;
+	virtual void SetHLODInEditorMaxDrawDistance(double InMaxDrawDistance) override;
+
+	virtual bool IsHLODInEditorAllowed(UWorld* InWorld, FText* OutDisallowedReason) const override;
+
 	/**
 	 * Convert the specified map to a world partition map.
 	 */
@@ -147,6 +161,7 @@ private:
 	void RunCommandletAsExternalProcess(const FString& InCommandletArgs, const FText& InOperationDescription, int32& OutResult, bool& bOutCancelled);
 	void OnConvertMap();
 
+	FDelegateHandle EditorInitializedHandle;
 	FDelegateHandle LevelEditorExtenderDelegateHandle;
 
 	TWeakPtr<SDockTab> WorldPartitionTab;
diff --git a/Engine/Source/Editor/WorldPartitionEditor/WorldPartitionEditor.Build.cs b/Engine/Source/Editor/WorldPartitionEditor/WorldPartitionEditor.Build.cs
--- a/Engine/Source/Editor/WorldPartitionEditor/WorldPartitionEditor.Build.cs
+++ b/Engine/Source/Editor/WorldPartitionEditor/WorldPartitionEditor.Build.cs
@@ -13,6 +13,7 @@ public class WorldPartitionEditor : ModuleRules
 				"CoreUObject",
 				"EditorFramework",
 				"EditorSubsystem",
+				"EditorWidgets",
 				"Engine",
 				"InputCore",
 				"Slate",
diff --git a/Engine/Source/Runtime/Engine/Private/WorldPartition/WorldPartition.cpp b/Engine/Source/Runtime/Engine/Private/WorldPartition/WorldPartition.cpp
--- a/Engine/Source/Runtime/Engine/Private/WorldPartition/WorldPartition.cpp
+++ b/Engine/Source/Runtime/Engine/Private/WorldPartition/WorldPartition.cpp
@@ -301,6 +301,7 @@ UWorldPartition::UWorldPartition(const FObjectInitializer& ObjectInitializer)
 	StreamingStateEpoch = 0;
 
 #if WITH_EDITOR
+	bAllowShowingHLODsInEditor = true;
 	WorldPartitionStreamingPolicyClass = UWorldPartitionLevelStreamingPolicy::StaticClass();
 #endif
 }
diff --git a/Engine/Source/Runtime/Engine/Public/WorldPartition/HLOD/HLODActor.h b/Engine/Source/Runtime/Engine/Public/WorldPartition/HLOD/HLODActor.h
--- a/Engine/Source/Runtime/Engine/Public/WorldPartition/HLOD/HLODActor.h
+++ b/Engine/Source/Runtime/Engine/Public/WorldPartition/HLOD/HLODActor.h
@@ -16,7 +16,7 @@ class UWorldPartitionHLODSourceActors;
 
 ENGINE_API DECLARE_LOG_CATEGORY_EXTERN(LogHLODHash, Log, All);
 
-UCLASS(NotPlaceable, MinimalAPI)
+UCLASS(NotPlaceable, MinimalAPI, HideCategories=(Rendering, Replication, Collision, Physics, Navigation, Networking, Input, Actor, LevelInstance, Cooking))
 class AWorldPartitionHLOD : public AActor
 {
 	GENERATED_UCLASS_BODY()
diff --git a/Engine/Source/Runtime/Engine/Public/WorldPartition/IWorldPartitionEditorModule.h b/Engine/Source/Runtime/Engine/Public/WorldPartition/IWorldPartitionEditorModule.h
--- a/Engine/Source/Runtime/Engine/Public/WorldPartition/IWorldPartitionEditorModule.h
+++ b/Engine/Source/Runtime/Engine/Public/WorldPartition/IWorldPartitionEditorModule.h
@@ -53,6 +53,20 @@ public:
 	virtual bool IsEditingContentBundle() const = 0;
 	virtual bool IsEditingContentBundle(const FGuid& ContentBundleGuid) const = 0;
 
+	virtual bool GetShowHLODsInEditor() const = 0;
+	virtual void SetShowHLODsInEditor(bool bInShowHLODsInEditor) = 0;
+
+	virtual bool GetShowHLODsOverLoadedRegions() const = 0;
+	virtual void SetShowHLODsOverLoadedRegions(bool bInShowHLODsOverLoadedRegions) = 0;
+
+	virtual double GetHLODInEditorMinDrawDistance() const = 0;
+	virtual void SetHLODInEditorMinDrawDistance(double InMinDrawDistance) = 0;
+
+	virtual double GetHLODInEditorMaxDrawDistance() const = 0;
+	virtual void SetHLODInEditorMaxDrawDistance(double InMaxDrawDistance) = 0;
+
+	virtual bool IsHLODInEditorAllowed(UWorld* InWorld, FText* OutDisallowedReason = nullptr) const = 0;
+
 	/** Triggered when a world is added. */
 	DECLARE_EVENT_OneParam(IWorldPartitionEditorModule, FWorldPartitionCreated, UWorld*);
 
diff --git a/Engine/Source/Runtime/Engine/Public/WorldPartition/WorldPartition.h b/Engine/Source/Runtime/Engine/Public/WorldPartition/WorldPartition.h
--- a/Engine/Source/Runtime/Engine/Public/WorldPartition/WorldPartition.h
+++ b/Engine/Source/Runtime/Engine/Public/WorldPartition/WorldPartition.h
@@ -392,6 +392,8 @@ public:
 
 	bool IsEnablingStreamingJustified() const { return bEnablingStreamingJustified; }
 
+	bool IsHLODsInEditorAllowed() const { return bAllowShowingHLODsInEditor; }
+
 	const TMap<FGuid, FDirtyActor>& GetDirtyActors() const { return DirtyActors; }
 #endif
 
@@ -497,6 +499,12 @@ private:
 	UPROPERTY(EditAnywhere, Category = WorldPartitionSetup, AdvancedDisplay)
 	EWorldPartitionDataLayersLogicOperator DataLayersLogicOperator;
 
+#if WITH_EDITORONLY_DATA
+	/** Whether HLODs should be allowed to be displayed in the editor for this map */
+	UPROPERTY(EditAnywhere, Category = WorldPartitionSetup, AdvancedDisplay, meta = (EditConditionHides, EditCondition = "bEnableStreaming", HideEditConditionToggle))
+	uint8 bAllowShowingHLODsInEditor : 1;
+#endif
+
 	TObjectPtr<UWorld> World;
 
 #if WITH_EDITOR
@@ -577,7 +585,7 @@ private:
 	ENGINE_API void HashActorDesc(FWorldPartitionActorDesc* ActorDesc);
 	ENGINE_API void UnhashActorDesc(FWorldPartitionActorDesc* ActorDesc);
 	void OnContentBundleRemovedContent(const FContentBundleEditor* ContentBundle);
-	bool IsStreamingEnabledInEditor() const;
+	ENGINE_API bool IsStreamingEnabledInEditor() const;
 
 public:
 	// Editor loader adapters management
