Fix duplicated LOCTEXT key "HLODInEditor_StreamingDisabled"
#rnx

diff --git a/Engine/Source/Editor/WorldPartitionEditor/Private/WorldPartitionEditorModule.cpp b/Engine/Source/Editor/WorldPartitionEditor/Private/WorldPartitionEditorModule.cpp
--- a/Engine/Source/Editor/WorldPartitionEditor/Private/WorldPartitionEditorModule.cpp
+++ b/Engine/Source/Editor/WorldPartitionEditor/Private/WorldPartitionEditorModule.cpp
@@ -465,7 +465,7 @@ bool FWorldPartitionEditorModule::IsHLODInEditorAllowed(UWorld* InWorld, FText*
 
 	if (!InWorld->GetWorldPartition()->IsHLODsInEditorAllowed())
 	{
-		SetDissallowedReason(LOCTEXT("HLODInEditor_StreamingDisabled", "HLOD in editor is disabled for this world"));
+		SetDissallowedReason(LOCTEXT("HLODInEditor_HLODsInEditorDisallowed", "HLOD in editor is disabled for this world"));
 		return false;
 	}
 
