Remove CleanupExternalObjectsEmptyFolders since it's now supposed to have zero impact with 29525768.

diff --git a/Engine/Config/BaseEditor.ini b/Engine/Config/BaseEditor.ini
--- a/Engine/Config/BaseEditor.ini
+++ b/Engine/Config/BaseEditor.ini
@@ -32,9 +32,6 @@ SdfFontAtlasSize=2048
 ; Content Sub Path being hidden from the Asset view for each mount
 ;+DenyListContentSubPaths=
 
-[WorldPartitionEditor]
-bCleanupExternalObjectsEmptyFolders=false
-
 [/Script/UnrealEd.ThumbnailManager]
 +RenderableThumbnailTypes=(ClassNeedingThumbnailName="/Script/CoreUObject.Class",RendererClassName="/Script/UnrealEd.ClassThumbnailRenderer")
 +RenderableThumbnailTypes=(ClassNeedingThumbnailName="/Script/Engine.Blueprint",RendererClassName="/Script/UnrealEd.BlueprintThumbnailRenderer")
diff --git a/Engine/Source/Editor/WorldPartitionEditor/Private/WorldPartitionEditorModule.cpp b/Engine/Source/Editor/WorldPartitionEditor/Private/WorldPartitionEditorModule.cpp
--- a/Engine/Source/Editor/WorldPartitionEditor/Private/WorldPartitionEditorModule.cpp
+++ b/Engine/Source/Editor/WorldPartitionEditor/Private/WorldPartitionEditorModule.cpp
@@ -2,7 +2,6 @@
 #include "WorldPartitionEditorModule.h"
 
 #include "WorldPartition/WorldPartition.h"
-#include "WorldPartition/WorldPartitionLog.h"
 #include "WorldPartition/WorldPartitionHLODsBuilder.h"
 #include "WorldPartition/WorldPartitionMiniMapBuilder.h"
 #include "WorldPartition/WorldPartitionLandscapeSplineMeshesBuilder.h"
@@ -28,14 +27,11 @@
 
 #include "Engine/Level.h"
 
-#include "HAL/FileManager.h"
-#include "Misc/ConfigCacheIni.h"
 #include "Misc/FileHelper.h"
 #include "Misc/MessageDialog.h"
 #include "Misc/StringBuilder.h"
 #include "Misc/ScopedSlowTask.h"
 #include "Internationalization/Regex.h"
-#include "ProfilingDebugging/ScopedTimers.h"
 
 #include "Interfaces/IMainFrameModule.h"
 #include "Widgets/SWindow.h"
@@ -213,8 +209,6 @@ void FWorldPartitionEditorModule::StartupModule()
 
 	FWorldPartitionClassDescRegistry().Get().Initialize();
 
-	CleanupExternalObjectsEmptyFolders();
-
 	EditorInitializedHandle = FEditorDelegates::OnEditorInitialized.AddLambda([this](double TimeToInitializeEditor)
 	{
 		// Register the Scene Outliner "World" filter category
@@ -275,8 +269,6 @@ void FWorldPartitionEditorModule::ShutdownModule()
 		PropertyEditor.UnregisterCustomClassLayout("WorldPartition");
 	}
 
-	WaitForCleanupExternalObjectsEmptyFolders();
-
 	FEditorDelegates::OnEditorInitialized.Remove(EditorInitializedHandle);
 }
 
@@ -947,112 +939,6 @@ void FWorldPartitionEditorModule::RegisterWorldPartitionLayout(FLayoutExtender&
 	Extender.ExtendLayout(FTabId("LevelEditorSelectionDetails"), ELayoutExtensionPosition::After, FTabManager::FTab(WorldPartitionEditorTabId, ETabState::ClosedTab));
 }
 
-void FWorldPartitionEditorModule::FCleanupExternalObjectsEmptyFoldersWorker::DoWork()
-{
-	auto CleanupEmptyFoldersRecursive = [](const TCHAR* LongRootPath)
-	{
-		FString RootPath = FPackageName::LongPackageNameToFilename(LongRootPath);
-
-		TMap<FString, uint32> FoldersFilesMap;
-		IFileManager::Get().IterateDirectoryRecursively(*RootPath, [&FoldersFilesMap](const TCHAR* FilenameOrDirectory, bool bIsDirectory)
-		{
-			if (bIsDirectory)
-			{
-				FoldersFilesMap.Add(FilenameOrDirectory, 0);
-			}
-			
-			// Register into parent folder
-			FoldersFilesMap.FindOrAdd(FPaths::GetPath(FilenameOrDirectory))++;
-
-			return true;
-		});
-
-		int32 NumFolders = 0;
-		int32 NumEmptyFolders = 0;
-		int32 NumDeletedFolders = 0;
-		int32 NumIteration = 0;
-
-		bool bShouldDeleteEmptyFolders = true;
-		while (bShouldDeleteEmptyFolders)
-		{
-			bShouldDeleteEmptyFolders = false;
-
-			for (auto& [Folder, NumFiles] : FoldersFilesMap)
-			{
-				if (!NumFiles &&  // Not yet deleted, and now empty
-					Folder != RootPath // Do not delete the RootPath, we should keep it and it would break the assertion that parent is in FolderFilesMap
-					)
-				{
-					if (IFileManager::Get().DeleteDirectory(*Folder, true))
-					{
-						NumDeletedFolders++;
-
-						// Register this folder as deleted to Avoid trying to delete it in the next iteration
-						NumFiles = INDEX_NONE;
-
-						// Unregister from parent folder
-						const int32 NumFilesInParentFolder = --FoldersFilesMap.FindChecked(FPaths::GetPath(Folder));
-						check(NumFilesInParentFolder >= 0);
-
-						bShouldDeleteEmptyFolders |= !NumFilesInParentFolder;
-					}
-
-					NumEmptyFolders++;
-				}
-
-				if (!NumIteration)
-				{
-					NumFolders++;
-				}
-			}
-
-			NumIteration++;
-		}
-
-		if (NumEmptyFolders)
-		{
-			TStringBuilder<1024> StringBuilder;
-			StringBuilder.Append(TEXT("\nExternal Objects Folder Stats\n"));
-			StringBuilder.Appendf(TEXT("\t           Path: %s\n"), LongRootPath);
-			StringBuilder.Appendf(TEXT("\t  Empty Folders: %d\n"), NumEmptyFolders);
-			StringBuilder.Appendf(TEXT("\tDeleted Folders: %d\n"), NumDeletedFolders);		
-			StringBuilder.Appendf(TEXT("\t    Num Folders: %d\n"), NumFolders);
-			UE_LOG(LogWorldPartition, Log, TEXT("%s"), StringBuilder.GetData());
-		}
-	};
-
-	const FString GameFolderPath = TEXT("/Game");
-	const FString ExternalObjectsFolderPath = GameFolderPath / FPackagePath::GetExternalObjectsFolderName();
-	const FString ExternalActorsFolderPath = GameFolderPath / FPackagePath::GetExternalActorsFolderName();
-
-	CleanupEmptyFoldersRecursive(*ExternalObjectsFolderPath);
-	CleanupEmptyFoldersRecursive(*ExternalActorsFolderPath);
-}
-
-void FWorldPartitionEditorModule::CleanupExternalObjectsEmptyFolders()
-{
-	bool bDoCleanup = false;
-	GConfig->GetBool(TEXT("WorldPartitionEditor"), TEXT("bCleanupExternalObjectsEmptyFolders"), bDoCleanup, GEditorIni);
-	if (bDoCleanup)
-	{
-		CleanupExternalObjectsEmptyFoldersWorkerAsyncTask = MakeUnique<FAsyncTask<FCleanupExternalObjectsEmptyFoldersWorker>>();
-		CleanupExternalObjectsEmptyFoldersWorkerAsyncTask->StartBackgroundTask();
-		// todo: cleanup content bundles when they are registered?
-
-		FCoreDelegates::OnPostEngineInit.AddLambda([this]() { WaitForCleanupExternalObjectsEmptyFolders(); });
-	}
-}
-
-void FWorldPartitionEditorModule::WaitForCleanupExternalObjectsEmptyFolders()
-{
-	if (CleanupExternalObjectsEmptyFoldersWorkerAsyncTask)
-	{
-		UE_SCOPED_TIMER(TEXT("FWorldPartitionEditorModule::WaitForCleanupExternalObjectsEmptyFolders"), LogWorldPartition, Display);
-		CleanupExternalObjectsEmptyFoldersWorkerAsyncTask->EnsureCompletion();
-		CleanupExternalObjectsEmptyFoldersWorkerAsyncTask.Reset();
-	}
-}
-
 FString UWorldPartitionConvertOptions::ToCommandletArgs() const
 {
 	TStringBuilder<1024> CommandletArgsBuilder;
diff --git a/Engine/Source/Editor/WorldPartitionEditor/Public/WorldPartitionEditorModule.h b/Engine/Source/Editor/WorldPartitionEditor/Public/WorldPartitionEditorModule.h
--- a/Engine/Source/Editor/WorldPartitionEditor/Public/WorldPartitionEditorModule.h
+++ b/Engine/Source/Editor/WorldPartitionEditor/Public/WorldPartitionEditorModule.h
@@ -2,7 +2,6 @@
 #pragma once
 
 #include "CoreMinimal.h"
-#include "Async/AsyncWork.h"
 #include "WorldPartition/IWorldPartitionEditorModule.h"
 
 class FTabManager;
@@ -140,12 +139,6 @@ private:
 	/** Inserts world partition tabs into the level editor layout */
 	void RegisterWorldPartitionLayout(FLayoutExtender& Extender);
 
-	/** Perform cleanup of empty folders for external objects paths */
-	void CleanupExternalObjectsEmptyFolders();
-
-	/** Wait for rmpty folders cleanup async task to complete */
-	void WaitForCleanupExternalObjectsEmptyFolders();
-
 	/** Spawns the world partition tab */
 	TSharedRef<SDockTab> SpawnWorldPartitionTab(const FSpawnTabArgs& Args);
 
@@ -174,17 +167,4 @@ private:
 	FOnPreExecuteCommandlet OnPreExecuteCommandletEvent;
 	FOnExecuteCommandlet OnExecuteCommandletEvent;
 	FOnPostExecuteCommandlet OnPostExecuteCommandletEvent;
-
-	class FCleanupExternalObjectsEmptyFoldersWorker : public FNonAbandonableTask
-	{
-	public:
-		void DoWork();
-
-		FORCEINLINE TStatId GetStatId() const
-		{
-			RETURN_QUICK_DECLARE_CYCLE_STAT(FCleanupExternalObjectsEmptyFoldersWorker, STATGROUP_ThreadPoolAsyncTasks);
-		}
-	};
-
-	TUniquePtr<FAsyncTask<FCleanupExternalObjectsEmptyFoldersWorker>> CleanupExternalObjectsEmptyFoldersWorkerAsyncTask;
 };
