World Partition Editor: show unsaved location volumes in the editor.

diff --git a/Engine/Source/Editor/WorldPartitionEditor/Private/WorldPartition/SWorldPartitionEditorGrid2D.cpp b/Engine/Source/Editor/WorldPartitionEditor/Private/WorldPartition/SWorldPartitionEditorGrid2D.cpp
--- a/Engine/Source/Editor/WorldPartitionEditor/Private/WorldPartition/SWorldPartitionEditorGrid2D.cpp
+++ b/Engine/Source/Editor/WorldPartitionEditor/Private/WorldPartition/SWorldPartitionEditorGrid2D.cpp
@@ -478,7 +478,12 @@ SWorldPartitionEditorGrid2D::SWorldPartitionEditorGrid2D()
 }
 
 SWorldPartitionEditorGrid2D::~SWorldPartitionEditorGrid2D()
-{}
+{
+	if (GEngine)
+	{
+		GEngine->OnLevelActorAdded().RemoveAll(this);
+	}
+}
 
 void SWorldPartitionEditorGrid2D::Construct(const FArguments& InArgs)
 {
@@ -556,6 +561,8 @@ void SWorldPartitionEditorGrid2D::Construct(const FArguments& InArgs)
 	];
 
 	BindCommands();
+
+	GEngine->OnLevelActorAdded().AddRaw(this, &SWorldPartitionEditorGrid2D::OnActorAdded);
 }
 
 void SWorldPartitionEditorGrid2D::BindCommands()
@@ -964,6 +971,14 @@ TSharedRef<SWidget> SWorldPartitionEditorGrid2D::GenerateContextualMenu() const
 	return UToolMenus::Get()->GenerateWidget(MenuName, FToolMenuContext(CommandList));
 }
 
+void SWorldPartitionEditorGrid2D::OnActorAdded(AActor* Actor)
+{
+	if (Actor->IsPackageExternal())
+	{
+		NewlyAddedUnsavedActors.Add(Actor);
+	}
+}
+
 FReply SWorldPartitionEditorGrid2D::OnMouseButtonUp(const FGeometry& MyGeometry, const FPointerEvent& MouseEvent)
 {
 	const bool bIsLeftMouseButtonEffecting = MouseEvent.GetEffectingButton() == EKeys::LeftMouseButton;
@@ -1263,19 +1278,39 @@ void SWorldPartitionEditorGrid2D::Tick(const FGeometry& AllottedGeometry, const
 		}
 	}, ForEachIntersectingActorParams);
 
+	// Add dirty actors as an acceleration for FWorldPartitionActorDescViewBoundsProxy constructor
+	for (auto& [Reference, Actor] : GetWorldPartition()->GetDirtyActors())
+	{
+		DirtyActorGuids.Add(Actor->GetActorGuid());
+	}
+
 	// Also include transient actor loader adapters that might have been spawned by blutilities, etc. Since these actors can't be saved because they are transient,
 	// they will never get an actor descriptor so they will never appear in the world partition editor. Also include unsaved, newly created actors for convenience.
-	for (auto& [Reference, Actor] : GetWorldPartition()->GetDirtyActors())
+	for (auto It = NewlyAddedUnsavedActors.CreateIterator(); It; ++It)
 	{
-		if (!GetWorldPartition()->GetActorDesc(Actor->GetActorGuid()) && Actor->Implements<UWorldPartitionActorLoaderInterface>())
+		if (It->IsValid())
 		{
-			if (IWorldPartitionActorLoaderInterface::ILoaderAdapter* LoaderAdapter = Cast<IWorldPartitionActorLoaderInterface>(Actor)->GetLoaderAdapter())
+			AActor* NewlyAddedUnsavedActor = It->Get();
+
+			if (!NewlyAddedUnsavedActor->GetPackage()->IsDirty())
 			{
-				ShownLoaderInterfaces.Add(Actor);
+				It.RemoveCurrent();
+			}
+			else
+			{
+				if (NewlyAddedUnsavedActor->Implements<UWorldPartitionActorLoaderInterface>())
+				{
+					if (IWorldPartitionActorLoaderInterface::ILoaderAdapter* LoaderAdapter = Cast<IWorldPartitionActorLoaderInterface>(NewlyAddedUnsavedActor)->GetLoaderAdapter())
+					{
+						ShownLoaderInterfaces.Add(NewlyAddedUnsavedActor);
+					}
+				}
 			}
 		}
-
-		DirtyActorGuids.Add(Actor->GetActorGuid());
+		else if (!It->IsValid(true))
+		{
+			It.RemoveCurrent();
+		}
 	}
 
 	FLoaderInterfaceSet LastHoveredLoaderInterfaces = MoveTemp(HoveredLoaderInterfaces);
diff --git a/Engine/Source/Editor/WorldPartitionEditor/Public/WorldPartition/SWorldPartitionEditorGrid2D.h b/Engine/Source/Editor/WorldPartitionEditor/Public/WorldPartition/SWorldPartitionEditorGrid2D.h
--- a/Engine/Source/Editor/WorldPartitionEditor/Public/WorldPartition/SWorldPartitionEditorGrid2D.h
+++ b/Engine/Source/Editor/WorldPartitionEditor/Public/WorldPartition/SWorldPartitionEditorGrid2D.h
@@ -105,6 +105,7 @@ protected:
 	
 	TSharedRef<SWidget> GenerateContextualMenu() const;
 
+	void OnActorAdded(AActor* Actor);
 
 	virtual void Tick(const FGeometry& AllottedGeometry, const double InCurrentTime, const float InDeltaTime) override;
 	virtual FReply OnMouseButtonDown(const FGeometry& MyGeometry, const FPointerEvent& MouseEvent) override;
@@ -226,4 +227,7 @@ protected:
 	TSharedPtr<SWorldPartitionViewportWidget> ViewportWidget;
 	TSharedPtr<STextBlock> TextWorldBoundsInKMWidget;
 	TSharedPtr<STextBlock> TextRulerWidget;
+
+	// List of actors without an actor descriptor
+	TSet<TWeakObjectPtr<AActor>> NewlyAddedUnsavedActors;
 };
\ No newline at end of file
