World Partition - HLODs in editor feature is now enabled by default

diff --git a/Engine/Source/Editor/WorldPartitionEditor/Private/WorldPartition/HLOD/HLODEditorSubsystem.cpp b/Engine/Source/Editor/WorldPartitionEditor/Private/WorldPartition/HLOD/HLODEditorSubsystem.cpp
--- a/Engine/Source/Editor/WorldPartitionEditor/Private/WorldPartition/HLOD/HLODEditorSubsystem.cpp
+++ b/Engine/Source/Editor/WorldPartitionEditor/Private/WorldPartition/HLOD/HLODEditorSubsystem.cpp
@@ -14,7 +14,7 @@
 
 static TAutoConsoleVariable<bool> CVarHLODInEditorEnabled(
 	TEXT("wp.Editor.HLOD.AllowShowingHLODsInEditor"),
-	false,
+	true,
 	TEXT("Allow showing World Partition HLODs in the editor."));
 
 
