World Partition Editor: show unsaved newly added actors in the editor view.

diff --git a/Engine/Source/Editor/WorldPartitionEditor/Private/WorldPartition/SWorldPartitionEditorGrid2D.cpp b/Engine/Source/Editor/WorldPartitionEditor/Private/WorldPartition/SWorldPartitionEditorGrid2D.cpp
--- a/Engine/Source/Editor/WorldPartitionEditor/Private/WorldPartition/SWorldPartitionEditorGrid2D.cpp
+++ b/Engine/Source/Editor/WorldPartitionEditor/Private/WorldPartition/SWorldPartitionEditorGrid2D.cpp
@@ -237,6 +237,21 @@ public:
 		return bUseActor ? ActorDesc->GetActor(false) : nullptr;
 	}
 
+	FName GetActorLabel() const
+	{
+		FName ActorLabel = ActorDesc->GetActorLabel();
+		
+		if (ActorLabel.IsNone() && bUseActor)
+		{
+			if (AActor* Actor = GetActor())
+			{
+				ActorLabel = *Actor->GetActorLabel(false);
+			}
+		}
+
+		return ActorLabel;
+	}
+
 	bool bUseActor;
 };
 
@@ -973,9 +988,10 @@ TSharedRef<SWidget> SWorldPartitionEditorGrid2D::GenerateContextualMenu() const
 
 void SWorldPartitionEditorGrid2D::OnActorAdded(AActor* Actor)
 {
-	if (Actor->IsPackageExternal())
+	if (!Actor->GetWorld()->IsGameWorld() && Actor->IsPackageExternal() && !Actor->IsChildActor())
 	{
-		NewlyAddedUnsavedActors.Add(Actor);
+		FWorldPartitionActorDesc* NewlyAddedUnsavedActorDesc = NewlyAddedUnsavedActorDescs.AddActor(Actor);
+		NewlyAddedUnsavedActors.Add(Actor, NewlyAddedUnsavedActorDesc);
 	}
 }
 
@@ -1288,12 +1304,13 @@ void SWorldPartitionEditorGrid2D::Tick(const FGeometry& AllottedGeometry, const
 	// they will never get an actor descriptor so they will never appear in the world partition editor. Also include unsaved, newly created actors for convenience.
 	for (auto It = NewlyAddedUnsavedActors.CreateIterator(); It; ++It)
 	{
-		if (It->IsValid())
+		if (It.Key().IsValid())
 		{
-			AActor* NewlyAddedUnsavedActor = It->Get();
+			AActor* NewlyAddedUnsavedActor = It.Key().Get();
 
 			if (!NewlyAddedUnsavedActor->GetPackage()->IsDirty())
 			{
+				NewlyAddedUnsavedActorDescs.RemoveActorDescriptor(It.Value());
 				It.RemoveCurrent();
 			}
 			else
@@ -1305,10 +1322,13 @@ void SWorldPartitionEditorGrid2D::Tick(const FGeometry& AllottedGeometry, const
 						ShownLoaderInterfaces.Add(NewlyAddedUnsavedActor);
 					}
 				}
+
+				ShownActorGuids.Add(NewlyAddedUnsavedActor->GetActorGuid());
 			}
 		}
-		else if (!It->IsValid(true))
+		else if (!It.Key().IsValid(true))
 		{
+			NewlyAddedUnsavedActorDescs.RemoveActorDescriptor(It.Value());
 			It.RemoveCurrent();
 		}
 	}
@@ -1414,6 +1434,10 @@ uint32 SWorldPartitionEditorGrid2D::PaintActors(const FGeometry& AllottedGeometr
 		{
 			ActorDescList.Emplace(ActorDesc, DirtyActorGuids.Contains(ActorGuid));
 		}
+		else if (const FWorldPartitionActorDesc* NewlyAddedUnsavedActorDesc = NewlyAddedUnsavedActorDescs.GetActorDesc(ActorGuid))
+		{
+			ActorDescList.Emplace(NewlyAddedUnsavedActorDesc, true);
+		}
 	}
 
 	auto DrawActorLabel = [this, &OutDrawElements, &LayerId, &AllottedGeometry](const FString& Label, const FVector2D& Pos, const FPaintGeometry& Geometry, const FLinearColor& Color, const FSlateFontInfo& Font, bool bShowBackground)
diff --git a/Engine/Source/Editor/WorldPartitionEditor/Public/WorldPartition/SWorldPartitionEditorGrid2D.h b/Engine/Source/Editor/WorldPartitionEditor/Public/WorldPartition/SWorldPartitionEditorGrid2D.h
--- a/Engine/Source/Editor/WorldPartitionEditor/Public/WorldPartition/SWorldPartitionEditorGrid2D.h
+++ b/Engine/Source/Editor/WorldPartitionEditor/Public/WorldPartition/SWorldPartitionEditorGrid2D.h
@@ -11,6 +11,7 @@
 #include "SWorldPartitionViewportWidget.h"
 #include "Widgets/Text/STextBlock.h"
 #include "SViewportToolBar.h"
+#include "WorldPartition/ActorDescList.h"
 #include "WorldPartition/WorldPartitionActorLoaderInterface.h"
 
 class SWorldPartitionEditorGrid2D : public SWorldPartitionEditorGrid
@@ -229,5 +230,6 @@ protected:
 	TSharedPtr<STextBlock> TextRulerWidget;
 
 	// List of actors without an actor descriptor
-	TSet<TWeakObjectPtr<AActor>> NewlyAddedUnsavedActors;
+	FActorDescList NewlyAddedUnsavedActorDescs;
+	TMap<TWeakObjectPtr<AActor>, FWorldPartitionActorDesc*> NewlyAddedUnsavedActors;
 };
\ No newline at end of file
