[Backout] - CL29887181
Original CL Desc
-----------------------------------------------------------------
World Partition Editor: show unsaved newly added actors in the editor view.

diff --git a/Engine/Source/Editor/WorldPartitionEditor/Private/WorldPartition/SWorldPartitionEditorGrid2D.cpp b/Engine/Source/Editor/WorldPartitionEditor/Private/WorldPartition/SWorldPartitionEditorGrid2D.cpp
--- a/Engine/Source/Editor/WorldPartitionEditor/Private/WorldPartition/SWorldPartitionEditorGrid2D.cpp
+++ b/Engine/Source/Editor/WorldPartitionEditor/Private/WorldPartition/SWorldPartitionEditorGrid2D.cpp
@@ -237,21 +237,6 @@ public:
 		return bUseActor ? ActorDesc->GetActor(false) : nullptr;
 	}
 
-	FName GetActorLabel() const
-	{
-		FName ActorLabel = ActorDesc->GetActorLabel();
-		
-		if (ActorLabel.IsNone() && bUseActor)
-		{
-			if (AActor* Actor = GetActor())
-			{
-				ActorLabel = *Actor->GetActorLabel(false);
-			}
-		}
-
-		return ActorLabel;
-	}
-
 	bool bUseActor;
 };
 
@@ -988,10 +973,9 @@ TSharedRef<SWidget> SWorldPartitionEditorGrid2D::GenerateContextualMenu() const
 
 void SWorldPartitionEditorGrid2D::OnActorAdded(AActor* Actor)
 {
-	if (!Actor->GetWorld()->IsGameWorld() && Actor->IsPackageExternal() && !Actor->IsChildActor())
+	if (Actor->IsPackageExternal())
 	{
-		FWorldPartitionActorDesc* NewlyAddedUnsavedActorDesc = NewlyAddedUnsavedActorDescs.AddActor(Actor);
-		NewlyAddedUnsavedActors.Add(Actor, NewlyAddedUnsavedActorDesc);
+		NewlyAddedUnsavedActors.Add(Actor);
 	}
 }
 
@@ -1304,13 +1288,12 @@ void SWorldPartitionEditorGrid2D::Tick(const FGeometry& AllottedGeometry, const
 	// they will never get an actor descriptor so they will never appear in the world partition editor. Also include unsaved, newly created actors for convenience.
 	for (auto It = NewlyAddedUnsavedActors.CreateIterator(); It; ++It)
 	{
-		if (It.Key().IsValid())
+		if (It->IsValid())
 		{
-			AActor* NewlyAddedUnsavedActor = It.Key().Get();
+			AActor* NewlyAddedUnsavedActor = It->Get();
 
 			if (!NewlyAddedUnsavedActor->GetPackage()->IsDirty())
 			{
-				NewlyAddedUnsavedActorDescs.RemoveActorDescriptor(It.Value());
 				It.RemoveCurrent();
 			}
 			else
@@ -1322,13 +1305,10 @@ void SWorldPartitionEditorGrid2D::Tick(const FGeometry& AllottedGeometry, const
 						ShownLoaderInterfaces.Add(NewlyAddedUnsavedActor);
 					}
 				}
-
-				ShownActorGuids.Add(NewlyAddedUnsavedActor->GetActorGuid());
 			}
 		}
-		else if (!It.Key().IsValid(true))
+		else if (!It->IsValid(true))
 		{
-			NewlyAddedUnsavedActorDescs.RemoveActorDescriptor(It.Value());
 			It.RemoveCurrent();
 		}
 	}
@@ -1434,10 +1414,6 @@ uint32 SWorldPartitionEditorGrid2D::PaintActors(const FGeometry& AllottedGeometr
 		{
 			ActorDescList.Emplace(ActorDesc, DirtyActorGuids.Contains(ActorGuid));
 		}
-		else if (const FWorldPartitionActorDesc* NewlyAddedUnsavedActorDesc = NewlyAddedUnsavedActorDescs.GetActorDesc(ActorGuid))
-		{
-			ActorDescList.Emplace(NewlyAddedUnsavedActorDesc, true);
-		}
 	}
 
 	auto DrawActorLabel = [this, &OutDrawElements, &LayerId, &AllottedGeometry](const FString& Label, const FVector2D& Pos, const FPaintGeometry& Geometry, const FLinearColor& Color, const FSlateFontInfo& Font, bool bShowBackground)
diff --git a/Engine/Source/Editor/WorldPartitionEditor/Public/WorldPartition/SWorldPartitionEditorGrid2D.h b/Engine/Source/Editor/WorldPartitionEditor/Public/WorldPartition/SWorldPartitionEditorGrid2D.h
--- a/Engine/Source/Editor/WorldPartitionEditor/Public/WorldPartition/SWorldPartitionEditorGrid2D.h
+++ b/Engine/Source/Editor/WorldPartitionEditor/Public/WorldPartition/SWorldPartitionEditorGrid2D.h
@@ -11,7 +11,6 @@
 #include "SWorldPartitionViewportWidget.h"
 #include "Widgets/Text/STextBlock.h"
 #include "SViewportToolBar.h"
-#include "WorldPartition/ActorDescList.h"
 #include "WorldPartition/WorldPartitionActorLoaderInterface.h"
 
 class SWorldPartitionEditorGrid2D : public SWorldPartitionEditorGrid
@@ -230,6 +229,5 @@ protected:
 	TSharedPtr<STextBlock> TextRulerWidget;
 
 	// List of actors without an actor descriptor
-	FActorDescList NewlyAddedUnsavedActorDescs;
-	TMap<TWeakObjectPtr<AActor>, FWorldPartitionActorDesc*> NewlyAddedUnsavedActors;
+	TSet<TWeakObjectPtr<AActor>> NewlyAddedUnsavedActors;
 };
\ No newline at end of file
