Bring back CL29887181 with fixes.

-----------------------------------------------------------------
[Backout] - CL29887181
Original CL Desc
-----------------------------------------------------------------
World Partition Editor: show unsaved newly added actors in the editor view.

diff --git a/Engine/Source/Editor/WorldPartitionEditor/Private/WorldPartition/SWorldPartitionEditorGrid2D.cpp b/Engine/Source/Editor/WorldPartitionEditor/Private/WorldPartition/SWorldPartitionEditorGrid2D.cpp
--- a/Engine/Source/Editor/WorldPartitionEditor/Private/WorldPartition/SWorldPartitionEditorGrid2D.cpp
+++ b/Engine/Source/Editor/WorldPartitionEditor/Private/WorldPartition/SWorldPartitionEditorGrid2D.cpp
@@ -34,6 +34,7 @@
 #include "WorldPartition/WorldPartitionLog.h"
 #include "WorldPartition/WorldPartitionActorDesc.h"
 #include "WorldPartition/WorldPartitionActorDescView.h"
+#include "WorldPartition/WorldPartitionActorDescUtils.h"
 #include "WorldPartition/WorldPartitionEditorHash.h"
 #include "WorldPartition/WorldPartitionEditorLoaderAdapter.h"
 #include "WorldPartition/WorldPartitionEditorPerProjectUserSettings.h"
@@ -237,6 +238,21 @@ public:
 		return bUseActor ? ActorDesc->GetActor(false) : nullptr;
 	}
 
+	FName GetActorLabel() const
+	{
+		FName ActorLabel = ActorDesc->GetActorLabel();
+		
+		if (ActorLabel.IsNone() && bUseActor)
+		{
+			if (AActor* Actor = GetActor())
+			{
+				ActorLabel = *Actor->GetActorLabel(false);
+			}
+		}
+
+		return ActorLabel;
+	}
+
 	bool bUseActor;
 };
 
@@ -973,9 +989,9 @@ TSharedRef<SWidget> SWorldPartitionEditorGrid2D::GenerateContextualMenu() const
 
 void SWorldPartitionEditorGrid2D::OnActorAdded(AActor* Actor)
 {
-	if (Actor->IsPackageExternal())
+	if (!Actor->GetWorld()->IsGameWorld() && Actor->IsPackageExternal() && !Actor->IsChildActor())
 	{
-		NewlyAddedUnsavedActors.Add(Actor);
+		NewlyAddedUnsavedActorDescs.OnActorAdded(Actor);
 	}
 }
 
@@ -1286,15 +1302,13 @@ void SWorldPartitionEditorGrid2D::Tick(const FGeometry& AllottedGeometry, const
 
 	// Also include transient actor loader adapters that might have been spawned by blutilities, etc. Since these actors can't be saved because they are transient,
 	// they will never get an actor descriptor so they will never appear in the world partition editor. Also include unsaved, newly created actors for convenience.
-	for (auto It = NewlyAddedUnsavedActors.CreateIterator(); It; ++It)
+	for (FActorDescList::TIterator<> ActorDescIterator(&NewlyAddedUnsavedActorDescs); ActorDescIterator; ++ActorDescIterator)
 	{
-		if (It->IsValid())
+		if (AActor* NewlyAddedUnsavedActor = ActorDescIterator->GetActor())
 		{
-			AActor* NewlyAddedUnsavedActor = It->Get();
-
 			if (!NewlyAddedUnsavedActor->GetPackage()->IsDirty())
 			{
-				It.RemoveCurrent();
+				ActorDescIterator.RemoveCurrent();
 			}
 			else
 			{
@@ -1305,11 +1319,13 @@ void SWorldPartitionEditorGrid2D::Tick(const FGeometry& AllottedGeometry, const
 						ShownLoaderInterfaces.Add(NewlyAddedUnsavedActor);
 					}
 				}
+
+				ShownActorGuids.Add(NewlyAddedUnsavedActor->GetActorGuid());
 			}
 		}
-		else if (!It->IsValid(true))
+		else if (!ActorDescIterator->GetActor(true))
 		{
-			It.RemoveCurrent();
+			ActorDescIterator.RemoveCurrent();
 		}
 	}
 
@@ -1414,6 +1430,10 @@ uint32 SWorldPartitionEditorGrid2D::PaintActors(const FGeometry& AllottedGeometr
 		{
 			ActorDescList.Emplace(ActorDesc, DirtyActorGuids.Contains(ActorGuid));
 		}
+		else if (const FWorldPartitionActorDesc* NewlyAddedUnsavedActorDesc = NewlyAddedUnsavedActorDescs.GetActorDesc(ActorGuid))
+		{
+			ActorDescList.Emplace(NewlyAddedUnsavedActorDesc, true);
+		}
 	}
 
 	auto DrawActorLabel = [this, &OutDrawElements, &LayerId, &AllottedGeometry](const FString& Label, const FVector2D& Pos, const FPaintGeometry& Geometry, const FLinearColor& Color, const FSlateFontInfo& Font, bool bShowBackground)
@@ -2316,4 +2336,16 @@ void SWorldPartitionEditorGrid2D::ClearSelection()
 	SelectBoxGridSnapped.Init();
 }
 
+void SWorldPartitionEditorGrid2D::FNewlyAddedUnsavedActorDescsDescRegistry::OnActorAdded(AActor* Actor)
+{
+	if (TUniquePtr<FWorldPartitionActorDesc>* ExistingActorDesc = GetActorDescriptor(Actor->GetActorGuid()))
+	{
+		FWorldPartitionActorDescUtils::UpdateActorDescriptorFromActor(Actor, *ExistingActorDesc);
+	}
+	else
+	{
+		AddActor(Actor);
+	}
+}
+
 #undef LOCTEXT_NAMESPACE
\ No newline at end of file
diff --git a/Engine/Source/Editor/WorldPartitionEditor/Public/WorldPartition/SWorldPartitionEditorGrid2D.h b/Engine/Source/Editor/WorldPartitionEditor/Public/WorldPartition/SWorldPartitionEditorGrid2D.h
--- a/Engine/Source/Editor/WorldPartitionEditor/Public/WorldPartition/SWorldPartitionEditorGrid2D.h
+++ b/Engine/Source/Editor/WorldPartitionEditor/Public/WorldPartition/SWorldPartitionEditorGrid2D.h
@@ -11,6 +11,7 @@
 #include "SWorldPartitionViewportWidget.h"
 #include "Widgets/Text/STextBlock.h"
 #include "SViewportToolBar.h"
+#include "WorldPartition/ActorDescList.h"
 #include "WorldPartition/WorldPartitionActorLoaderInterface.h"
 
 class SWorldPartitionEditorGrid2D : public SWorldPartitionEditorGrid
@@ -229,5 +230,11 @@ protected:
 	TSharedPtr<STextBlock> TextRulerWidget;
 
 	// List of actors without an actor descriptor
-	TSet<TWeakObjectPtr<AActor>> NewlyAddedUnsavedActors;
+	class FNewlyAddedUnsavedActorDescsDescRegistry : public FActorDescList
+	{
+	public:
+		void OnActorAdded(AActor* Actor);
+	};
+
+	FNewlyAddedUnsavedActorDescsDescRegistry NewlyAddedUnsavedActorDescs;
 };
\ No newline at end of file
diff --git a/Engine/Source/Runtime/Engine/Public/WorldPartition/ActorDescList.h b/Engine/Source/Runtime/Engine/Public/WorldPartition/ActorDescList.h
--- a/Engine/Source/Runtime/Engine/Public/WorldPartition/ActorDescList.h
+++ b/Engine/Source/Runtime/Engine/Public/WorldPartition/ActorDescList.h
@@ -20,6 +20,7 @@ class FActorDescList
 	};
 
 	using FGuidActorDescMap = TMap<FGuid, TUniquePtr<FWorldPartitionActorDesc>*, FDefaultSetAllocator, FActorGuidKeyFuncs>;
+	using FActorDescArray = TChunkedArray<TUniquePtr<FWorldPartitionActorDesc>>;
 
 public:
 	FActorDescList() {}
@@ -69,9 +70,9 @@ public:
 		}
 
 		/**
-		 * Iterates to next suitable actor desc.
+		 * Iterates to next suitable actor desc
 		 */
-		void operator++()
+		FORCEINLINE void operator++()
 		{
 			do
 			{
@@ -79,6 +80,14 @@ public:
 			} while (ShouldSkip());
 		}
 
+		/**
+		 * Removes the current iterator element
+		 */
+		FORCEINLINE void RemoveCurrent()
+		{
+			ActorsIterator.RemoveCurrent();
+		}
+
 		/**
 		 * Returns the current suitable actor desc pointed at by the Iterator
 		 *
@@ -111,7 +120,7 @@ public:
 		}
 
 		/**
-		 * Returns the actor class on which the iterator iterates on.
+		 * Returns the actor class on which the iterator iterates on
 		 *
 		 * @return the actor class
 		 */
@@ -119,7 +128,7 @@ public:
 
 	protected:
 		/**
-		 * Determines whether the iterator currently points to a valid actor desc or not.
+		 * Determines whether the iterator currently points to a valid actor desc or not
 		 * @return true if we should skip the actor desc
 		 */
 		FORCEINLINE bool ShouldSkip() const
@@ -162,11 +171,9 @@ public:
 	ENGINE_API void RemoveActorDescriptor(FWorldPartitionActorDesc* ActorDesc);
 
 protected:
-
 	ENGINE_API TUniquePtr<FWorldPartitionActorDesc>* GetActorDescriptor(const FGuid& ActorGuid);
 
-	TChunkedArray<TUniquePtr<FWorldPartitionActorDesc>> ActorDescList;
-
+	FActorDescArray ActorDescList;
 	FGuidActorDescMap ActorsByGuid;
 #endif
 };
