World Partition - HLOD: Save HLOD in editor settings when changing them through the viewport menus

diff --git a/Engine/Source/Editor/WorldPartitionEditor/Private/WorldPartitionEditorModule.cpp b/Engine/Source/Editor/WorldPartitionEditor/Private/WorldPartitionEditorModule.cpp
--- a/Engine/Source/Editor/WorldPartitionEditor/Private/WorldPartitionEditorModule.cpp
+++ b/Engine/Source/Editor/WorldPartitionEditor/Private/WorldPartitionEditorModule.cpp
@@ -394,7 +394,11 @@ bool FWorldPartitionEditorModule::GetShowHLODsInEditor() const
 
 void FWorldPartitionEditorModule::SetShowHLODsInEditor(bool bInShowHLODsInEditor)
 {
-	GetMutableDefault<UWorldPartitionEditorSettings>()->bShowHLODsInEditor = bInShowHLODsInEditor;
+	if (GetMutableDefault<UWorldPartitionEditorSettings>()->bShowHLODsInEditor != bInShowHLODsInEditor)
+	{
+		GetMutableDefault<UWorldPartitionEditorSettings>()->bShowHLODsInEditor = bInShowHLODsInEditor;
+		GetMutableDefault<UWorldPartitionEditorSettings>()->SaveConfig();
+	}
 }
 
 bool FWorldPartitionEditorModule::GetShowHLODsOverLoadedRegions() const
@@ -404,7 +408,11 @@ bool FWorldPartitionEditorModule::GetShowHLODsOverLoadedRegions() const
 
 void FWorldPartitionEditorModule::SetShowHLODsOverLoadedRegions(bool bInShowHLODsOverLoadedRegions)
 {
-	GetMutableDefault<UWorldPartitionEditorSettings>()->bShowHLODsOverLoadedRegions = bInShowHLODsOverLoadedRegions;
+	if (GetMutableDefault<UWorldPartitionEditorSettings>()->bShowHLODsOverLoadedRegions != bInShowHLODsOverLoadedRegions)
+	{
+		GetMutableDefault<UWorldPartitionEditorSettings>()->bShowHLODsOverLoadedRegions = bInShowHLODsOverLoadedRegions;
+		GetMutableDefault<UWorldPartitionEditorSettings>()->SaveConfig();
+	}
 }
 
 double FWorldPartitionEditorModule::GetHLODInEditorMinDrawDistance() const
@@ -414,7 +422,11 @@ double FWorldPartitionEditorModule::GetHLODInEditorMinDrawDistance() const
 
 void FWorldPartitionEditorModule::SetHLODInEditorMinDrawDistance(double InMinDrawDistance)
 {
-	GetMutableDefault<UWorldPartitionEditorSettings>()->HLODMinDrawDistance = InMinDrawDistance;
+	if (GetMutableDefault<UWorldPartitionEditorSettings>()->HLODMinDrawDistance != InMinDrawDistance)
+	{
+		GetMutableDefault<UWorldPartitionEditorSettings>()->HLODMinDrawDistance = InMinDrawDistance;
+		GetMutableDefault<UWorldPartitionEditorSettings>()->SaveConfig();
+	}
 }
 
 double FWorldPartitionEditorModule::GetHLODInEditorMaxDrawDistance() const
@@ -424,7 +436,11 @@ double FWorldPartitionEditorModule::GetHLODInEditorMaxDrawDistance() const
 
 void FWorldPartitionEditorModule::SetHLODInEditorMaxDrawDistance(double InMaxDrawDistance)
 {
-	GetMutableDefault<UWorldPartitionEditorSettings>()->HLODMaxDrawDistance = InMaxDrawDistance;
+	if (GetMutableDefault<UWorldPartitionEditorSettings>()->HLODMaxDrawDistance != InMaxDrawDistance)
+	{
+		GetMutableDefault<UWorldPartitionEditorSettings>()->HLODMaxDrawDistance = InMaxDrawDistance;
+		GetMutableDefault<UWorldPartitionEditorSettings>()->SaveConfig();
+	}
 }
 
 bool FWorldPartitionEditorModule::IsHLODInEditorAllowed(UWorld* InWorld, FText* OutDisallowedReason) const
