World Partition - HLOD in editor: Properly handle container registration/unregistration

diff --git a/Engine/Source/Editor/WorldPartitionEditor/Private/WorldPartition/HLOD/HLODEditorData.cpp b/Engine/Source/Editor/WorldPartitionEditor/Private/WorldPartition/HLOD/HLODEditorData.cpp
--- a/Engine/Source/Editor/WorldPartitionEditor/Private/WorldPartition/HLOD/HLODEditorData.cpp
+++ b/Engine/Source/Editor/WorldPartitionEditor/Private/WorldPartition/HLOD/HLODEditorData.cpp
@@ -7,6 +7,7 @@
 #include "WorldPartition/HLOD/HLODLoaderAdapter.h"
 
 #include "WorldPartition/ActorDescContainerInstanceCollection.h"
+#include "WorldPartition/LoaderAdapter/LoaderAdapterActorList.h"
 #include "WorldPartition/WorldPartitionActorDescInstance.h"
 #include "WorldPartition/WorldPartition.h"
 #include "WorldPartition/WorldPartitionHandle.h"
@@ -16,27 +17,55 @@ FWorldPartitionHLODEditorData::FWorldPartitionHLODEditorData(UWorldPartition* In
 	: WorldPartition(InWorldPartition)
 	, LastStateUpdate(INDEX_NONE)
 {
+	WorldPartition->OnActorDescContainerInstanceRegistered.AddRaw(this, &FWorldPartitionHLODEditorData::OnActorDescContainerInstanceRegistered);
+	WorldPartition->OnActorDescContainerInstanceUnregistered.AddRaw(this, &FWorldPartitionHLODEditorData::OnActorDescContainerInstanceUnregistered);
+
+	// Since this is created upon WP init, we missed the first broadcasts for existing container instances. Register them manually.
+	WorldPartition->ForEachActorDescContainer([this](UActorDescContainerInstance* InContainerInstance)
+	{
+		OnActorDescContainerInstanceRegistered(InContainerInstance);
+	});	
+}
+
+void FWorldPartitionHLODEditorData::OnActorDescContainerInstanceRegistered(UActorDescContainerInstance* InContainerInstance)
+{
+	check(!PerContainerInstanceHLODActorDataMap.Contains(InContainerInstance));
+
+	FContainerInstanceHLODActorData* ContainerInstanceHLODActorData = nullptr;
+
 	// First pass, build mapping of GUID to HLOD scene node
-	for (FActorDescContainerInstanceCollection::TIterator<AWorldPartitionHLOD> HLODIterator(WorldPartition); HLODIterator; ++HLODIterator)
+	for (UActorDescContainerInstance::TConstIterator<AWorldPartitionHLOD> HLODIterator(InContainerInstance); HLODIterator; ++HLODIterator)
 	{
+		// Only create an entry in the map if there are HLOD actors in this container instance.
+		if (!ContainerInstanceHLODActorData)
+		{
+			ContainerInstanceHLODActorData = &PerContainerInstanceHLODActorDataMap.Emplace(InContainerInstance);
+		} 
+
 		const FHLODActorDesc& HLODActorDesc = *(FHLODActorDesc*)HLODIterator->GetActorDesc();
 		const FGuid& HLODActorGuid = HLODActorDesc.GetGuid();
 
-		TUniquePtr<FHLODSceneNode>& HLODSceneNode = HLODActorNodes.Emplace(HLODActorGuid, new FHLODSceneNode());
+		TUniquePtr<FHLODSceneNode>& HLODSceneNode = ContainerInstanceHLODActorData->HLODActorNodes.Emplace(HLODActorGuid, new FHLODSceneNode());
 		HLODSceneNode->Bounds = HLODActorDesc.GetEditorBounds();
-		HLODSceneNode->HLODActorHandle = FWorldPartitionHandle(InWorldPartition, HLODActorDesc.GetGuid());
+		HLODSceneNode->HLODActorHandle = FWorldPartitionHandle(WorldPartition, HLODActorDesc.GetGuid());
+	}
+
+	// If there are no HLOD actors in this container instance, early out
+	if (!ContainerInstanceHLODActorData)
+	{
+		return;
 	}
 
 	// Second pass, build a hierarchy now that nodes were all created
-	for (FActorDescContainerInstanceCollection::TIterator<AWorldPartitionHLOD> HLODIterator(WorldPartition); HLODIterator; ++HLODIterator)
+	for (UActorDescContainerInstance::TConstIterator<AWorldPartitionHLOD> HLODIterator(InContainerInstance); HLODIterator; ++HLODIterator)
 	{
 		const FHLODActorDesc& HLODActorDesc = *(FHLODActorDesc*)HLODIterator->GetActorDesc();
 		const FGuid& HLODActorGuid = HLODActorDesc.GetGuid();
-		TUniquePtr<FHLODSceneNode>& HLODSceneNode = HLODActorNodes.FindChecked(HLODActorGuid);
+		TUniquePtr<FHLODSceneNode>& HLODSceneNode = ContainerInstanceHLODActorData->HLODActorNodes.FindChecked(HLODActorGuid);
 
 		for (const FGuid& ChildHLODActorGuid : HLODActorDesc.GetChildHLODActors())
 		{
-			if (TUniquePtr<FHLODSceneNode>* ChildHLODSceneNodePtr = HLODActorNodes.Find(ChildHLODActorGuid))
+			if (TUniquePtr<FHLODSceneNode>* ChildHLODSceneNodePtr = ContainerInstanceHLODActorData->HLODActorNodes.Find(ChildHLODActorGuid))
 			{
 				FHLODSceneNode* ChildHLODSceneNode = ChildHLODSceneNodePtr->Get();
 				ChildHLODSceneNode->ParentHLOD = HLODSceneNode.Get();
@@ -47,15 +76,20 @@ FWorldPartitionHLODEditorData::FWorldPartitionHLODEditorData(UWorldPartition* In
 	}
 
 	// Cache top level HLOD nodes for a faster iteration during the subsystem tick.
-	for (const auto& [HLODActorGuid, HLODActorNode] : HLODActorNodes)
+	for (const auto& [HLODActorGuid, HLODActorNode] : ContainerInstanceHLODActorData->HLODActorNodes)
 	{
 		if (HLODActorNode->ParentHLOD == nullptr)
 		{
-			TopLevelHLODActorNodes.Add(HLODActorNode.Get());
+			ContainerInstanceHLODActorData->TopLevelHLODActorNodes.Add(HLODActorNode.Get());
 		}
 	}
 }
 
+void FWorldPartitionHLODEditorData::OnActorDescContainerInstanceUnregistered(UActorDescContainerInstance* InContainerInstance)
+{
+	PerContainerInstanceHLODActorDataMap.Remove(InContainerInstance);
+}
+
 struct FBoundsWithVolume
 {
 	FBoundsWithVolume(const FBox& InBox)
@@ -189,9 +223,14 @@ void FWorldPartitionHLODEditorData::UpdateLoadedActorsState()
 	};
 
 	// Update Nodes, starting from the top level HLODs down to their children
-	for (FHLODSceneNode* HLODSceneNode : TopLevelHLODActorNodes)
+	for (auto& MapEntry : PerContainerInstanceHLODActorDataMap)
 	{
-		UpdateNodeState(HLODSceneNode);
+		FContainerInstanceHLODActorData& ContainerInstanceHLODActorData = MapEntry.Value;
+
+		for (FHLODSceneNode* HLODSceneNode : ContainerInstanceHLODActorData.TopLevelHLODActorNodes)
+		{
+			UpdateNodeState(HLODSceneNode);
+		}
 	}
 }
 
@@ -205,16 +244,21 @@ void FWorldPartitionHLODEditorData::UpdateVisibility(const FVector& InCameraLoca
 		return;
 	}
 
-	// For each top level HLOD actor (ex: HLOD2 in a 3 level of HLOD setup)
-	// Determine visibility of each HLOD by a few factors:
-	//  * If source (non-hlod) actors are loaded beneat it - HIDDEN
-	//  * If near culled (using the min visible distance) - HIDDEN
-	// Then if a given HLOD is HIDDEN, perform the same logic for its children.
-	// If the HLOD is VISIBLE, then flag all it's children as HIDDEN
-	for (const auto HLODActorNode : TopLevelHLODActorNodes)
+	for (auto& MapEntry : PerContainerInstanceHLODActorDataMap)
 	{
-		// Recurse from the top level HLOD down to HLOD0
-		HLODActorNode->UpdateVisibility(InCameraLocation, InMinDrawDistance, InMaxDrawDistance, /*bInForceHidden*/false, bInForceVisibilityUpdate, LastStateUpdate);
+		FContainerInstanceHLODActorData& ContainerInstanceHLODActorData = MapEntry.Value;
+		
+		// For each top level HLOD actor (ex: HLOD2 in a 3 level of HLOD setup)
+		// Determine visibility of each HLOD by a few factors:
+		//  * If source (non-hlod) actors are loaded beneat it - HIDDEN
+		//  * If near culled (using the min visible distance) - HIDDEN
+		// Then if a given HLOD is HIDDEN, perform the same logic for its children.
+		// If the HLOD is VISIBLE, then flag all it's children as HIDDEN
+		for (const auto HLODActorNode : ContainerInstanceHLODActorData.TopLevelHLODActorNodes)
+		{
+			// Recurse from the top level HLOD down to HLOD0
+			HLODActorNode->UpdateVisibility(InCameraLocation, InMinDrawDistance, InMaxDrawDistance, /*bInForceHidden*/false, bInForceVisibilityUpdate, LastStateUpdate);
+		}
 	}
 }
 
diff --git a/Engine/Source/Editor/WorldPartitionEditor/Private/WorldPartition/HLOD/HLODLoaderAdapter.cpp b/Engine/Source/Editor/WorldPartitionEditor/Private/WorldPartition/HLOD/HLODLoaderAdapter.cpp
--- a/Engine/Source/Editor/WorldPartitionEditor/Private/WorldPartition/HLOD/HLODLoaderAdapter.cpp
+++ b/Engine/Source/Editor/WorldPartitionEditor/Private/WorldPartition/HLOD/HLODLoaderAdapter.cpp
@@ -15,20 +15,23 @@
 
 
 FLoaderAdapterHLOD::FLoaderAdapterHLOD(UWorld* InWorld)
-	: FLoaderAdapterActorList(InWorld)
+	: ILoaderAdapter(InWorld)
 {
-	UWorldPartition* WorldPartition = InWorld->GetWorldPartition();
+	Load();
+}
+
+void FLoaderAdapterHLOD::ForEachActor(TFunctionRef<void(const FWorldPartitionHandle&)> InOperation) const
+{
+	UWorldPartition* WorldPartition = GetWorld()->GetWorldPartition();
 
 	for (FActorDescContainerInstanceCollection::TIterator<AWorldPartitionHLOD> HLODIterator(WorldPartition); HLODIterator; ++HLODIterator)
 	{
 		const FHLODActorDesc& HLODActorDesc = *(FHLODActorDesc*)HLODIterator->GetActorDesc();
 		if (ShouldLoadHLOD(HLODActorDesc))
 		{
-			Actors.Add(FWorldPartitionHandle(WorldPartition, HLODActorDesc.GetGuid()));
+			InOperation(FWorldPartitionHandle(WorldPartition, HLODActorDesc.GetGuid()));
 		}
 	}
-
-	RefreshLoadedState();
 }
 
 // TODO - How to choose which HLODs are loaded ?
@@ -38,7 +41,7 @@ FLoaderAdapterHLOD::FLoaderAdapterHLOD(UWorld* InWorld)
 // * HLOD Layer setting?
 //
 // Current solution is to load always loaded HLOD + any HLOD not built from instancing
-bool FLoaderAdapterHLOD::ShouldLoadHLOD(const FHLODActorDesc& HLODActorDesc)
+bool FLoaderAdapterHLOD::ShouldLoadHLOD(const FHLODActorDesc& HLODActorDesc) const
 {
 	FSoftObjectPath HLODLayerPath(HLODActorDesc.GetSourceHLODLayer());
 	if (UHLODLayer* HLODLayer = Cast<UHLODLayer>(HLODLayerPath.TryLoad()))
@@ -54,5 +57,6 @@ bool FLoaderAdapterHLOD::ShouldLoadHLOD(const FHLODActorDesc& HLODActorDesc)
 
 bool FLoaderAdapterHLOD::PassActorDescFilter(const FWorldPartitionHandle& ActorHandle) const
 {
-	return ActorHandle.IsValid() && !ActorsToRemove.Contains(ActorHandle);
+	// Avoid the base class implementation which will skip actors which are not editor relevant
+	return true;
 }
diff --git a/Engine/Source/Editor/WorldPartitionEditor/Public/WorldPartition/HLOD/HLODEditorData.h b/Engine/Source/Editor/WorldPartitionEditor/Public/WorldPartition/HLOD/HLODEditorData.h
--- a/Engine/Source/Editor/WorldPartitionEditor/Public/WorldPartition/HLOD/HLODEditorData.h
+++ b/Engine/Source/Editor/WorldPartitionEditor/Public/WorldPartition/HLOD/HLODEditorData.h
@@ -2,10 +2,10 @@
 
 #pragma once
 
-#include "WorldPartition/LoaderAdapter/LoaderAdapterActorList.h"
-
+#include "WorldPartition/HLOD/HLODLoaderAdapter.h"
 
 class UWorldPartition;
+class UActorDescContainerInstance;
 
 // Represent an HLOD actor in the editor, loaded or not
 struct FHLODSceneNode
@@ -36,9 +36,18 @@ public:
 	void SetHLODLoadingState(bool bInShouldBeLoaded);
 
 private:
+	void OnActorDescContainerInstanceRegistered(UActorDescContainerInstance* InContainerInstance);
+	void OnActorDescContainerInstanceUnregistered(UActorDescContainerInstance* InContainerInstance);
+
+private:
+	struct FContainerInstanceHLODActorData
+	{
+		TMap<FGuid, TUniquePtr<FHLODSceneNode>> HLODActorNodes;
+		TArray<FHLODSceneNode*> TopLevelHLODActorNodes;
+	};
+
 	UWorldPartition* WorldPartition;
-	TMap<FGuid, TUniquePtr<FHLODSceneNode>> HLODActorNodes;
-	TArray<FHLODSceneNode*> TopLevelHLODActorNodes;
-	TUniquePtr<FLoaderAdapterActorList> HLODActorsLoader;
+	TMap<UActorDescContainerInstance*, FContainerInstanceHLODActorData> PerContainerInstanceHLODActorDataMap;
+	TUniquePtr<FLoaderAdapterHLOD> HLODActorsLoader;
 	int32 LastStateUpdate;
 };
diff --git a/Engine/Source/Editor/WorldPartitionEditor/Public/WorldPartition/HLOD/HLODLoaderAdapter.h b/Engine/Source/Editor/WorldPartitionEditor/Public/WorldPartition/HLOD/HLODLoaderAdapter.h
--- a/Engine/Source/Editor/WorldPartitionEditor/Public/WorldPartition/HLOD/HLODLoaderAdapter.h
+++ b/Engine/Source/Editor/WorldPartitionEditor/Public/WorldPartition/HLOD/HLODLoaderAdapter.h
@@ -2,17 +2,21 @@
 
 #pragma once
 
-#include "WorldPartition/LoaderAdapter/LoaderAdapterActorList.h"
+#include "WorldPartition/WorldPartitionActorLoaderInterface.h"
 
 class FHLODActorDesc;
 
-class FLoaderAdapterHLOD : public FLoaderAdapterActorList
+class FLoaderAdapterHLOD : public IWorldPartitionActorLoaderInterface::ILoaderAdapter
 {
 public:
 	FLoaderAdapterHLOD(UWorld* InWorld);
 
+protected:
+	//~ Begin ILoaderAdapter interface
+	virtual void ForEachActor(TFunctionRef<void(const FWorldPartitionHandle&)> InOperation) const override;
 	virtual bool PassActorDescFilter(const FWorldPartitionHandle& ActorHandle) const override;
+	//~ End ILoaderAdapter interface
 
 private:
-	bool ShouldLoadHLOD(const FHLODActorDesc& HLODActorDesc);	
+	bool ShouldLoadHLOD(const FHLODActorDesc& HLODActorDesc) const;	
 };
