World Partition - HLOD: Moved the HLODEditorData.h header to the Private folder as this doesn't need to be exposed
#rnx

diff --git a/Engine/Source/Editor/WorldPartitionEditor/Private/WorldPartition/HLOD/HLODEditorData.h b/Engine/Source/Editor/WorldPartitionEditor/Private/WorldPartition/HLOD/HLODEditorData.h
new file mode 100644
--- /dev/null
+++ b/Engine/Source/Editor/WorldPartitionEditor/Private/WorldPartition/HLOD/HLODEditorData.h
@@ -0,0 +1,53 @@
+// Copyright Epic Games, Inc. All Rights Reserved.
+
+#pragma once
+
+#include "WorldPartition/HLOD/HLODLoaderAdapter.h"
+
+class UWorldPartition;
+class UActorDescContainerInstance;
+
+// Represent an HLOD actor in the editor, loaded or not
+struct FHLODSceneNode
+{
+	void UpdateVisibility(const FVector& InCameraLocation, double InMinDrawDistance, double InMaxDrawDistance, bool bInForceHidden, bool bInForceVisibilityUpdate, int32 InLastStateUpdate);
+
+	FHLODSceneNode* ParentHLOD = nullptr;
+	TArray<FHLODSceneNode*> ChildrenHLODs;
+
+
+	FBoxSphereBounds Bounds;
+	bool bCachedIsVisible = true;
+	FWorldPartitionHandle HLODActorHandle;
+	
+	int32 HasIntersectingLoadedRegion = INDEX_NONE;
+};
+
+// Editor state of HLODs for a given World Partition
+struct FWorldPartitionHLODEditorData
+{
+public:
+	FWorldPartitionHLODEditorData(UWorldPartition* InWorldPartition);
+	
+	void ClearLoadedActorsState();
+	void UpdateLoadedActorsState();
+	void UpdateVisibility(const FVector& InCameraLocation, double InMinDrawDistance, double InMaxDrawDistance, bool bForceVisibilityUpdate);
+
+	void SetHLODLoadingState(bool bInShouldBeLoaded);
+
+private:
+	void OnActorDescContainerInstanceRegistered(UActorDescContainerInstance* InContainerInstance);
+	void OnActorDescContainerInstanceUnregistered(UActorDescContainerInstance* InContainerInstance);
+
+private:
+	struct FContainerInstanceHLODActorData
+	{
+		TMap<FGuid, TUniquePtr<FHLODSceneNode>> HLODActorNodes;
+		TArray<FHLODSceneNode*> TopLevelHLODActorNodes;
+	};
+
+	UWorldPartition* WorldPartition;
+	TMap<UActorDescContainerInstance*, FContainerInstanceHLODActorData> PerContainerInstanceHLODActorDataMap;
+	TUniquePtr<FLoaderAdapterHLOD> HLODActorsLoader;
+	int32 LastStateUpdate;
+};
diff --git a/Engine/Source/Editor/WorldPartitionEditor/Public/WorldPartition/HLOD/HLODEditorData.h b/Engine/Source/Editor/WorldPartitionEditor/Public/WorldPartition/HLOD/HLODEditorData.h
deleted file mode 100644
