TDirtyExternalActorTracker: tracks dirty external actors.

World Partition Editor:
- Abstract actors bounds rendering to allow rendering dirty and unsaved actors without the need of actor descriptors.
- Update dirty actors code management to TDirtyExternalActorTracker.

World Partition:
- Update dirty actors management code to TDirtyExternalActorTracker.

diff --git a/Engine/Source/Editor/WorldPartitionEditor/Private/WorldPartition/SWorldPartitionEditorGrid2D.cpp b/Engine/Source/Editor/WorldPartitionEditor/Private/WorldPartition/SWorldPartitionEditorGrid2D.cpp
--- a/Engine/Source/Editor/WorldPartitionEditor/Private/WorldPartition/SWorldPartitionEditorGrid2D.cpp
+++ b/Engine/Source/Editor/WorldPartitionEditor/Private/WorldPartition/SWorldPartitionEditorGrid2D.cpp
@@ -58,73 +58,6 @@ static FAutoConsoleCommand CVarToggleShowEditorProfilingStats(
 	FConsoleCommandDelegate::CreateLambda([] { GShowEditorProfilingStats = !GShowEditorProfilingStats; })
 );
 
-class FEditorGrid2DActorDescInstance : public FWorldPartitionActorDescInstance
-{
-public:
-	FEditorGrid2DActorDescInstance(){}
-	virtual ~FEditorGrid2DActorDescInstance() {}
-
-	virtual FBox GetEditorBounds() const override
-	{
-		if (AActor* Actor = GetActor(false))
-		{
-			return Actor->GetStreamingBounds();
-		}
-
-		return FWorldPartitionActorDescInstance::GetEditorBounds();
-	}
-
-	virtual FName GetActorLabel() const override
-	{
-		FName ActorLabel = FWorldPartitionActorDescInstance::GetActorLabel();
-
-		if (ActorLabel.IsNone())
-		{
-			if (AActor* Actor = GetActor(false))
-			{
-				ActorLabel = *Actor->GetActorLabel(false);
-			}
-		}
-
-		return ActorLabel;
-	}
-};
-
-class FEditorGrid2DUnsavedActorDescInstance : public FEditorGrid2DActorDescInstance
-{
-public:
-	FEditorGrid2DUnsavedActorDescInstance(AActor* InActor)
-	{
-		ActorDescPtr = InActor->CreateActorDesc();
-		ActorDesc = ActorDescPtr.Get();
-	}
-
-	virtual ~FEditorGrid2DUnsavedActorDescInstance() {}
-	
-	// No need for container registration here
-	virtual void RegisterChildContainerInstance() override {}
-	virtual void UnregisterChildContainerInstance() override {}
-
-	TUniquePtr<FWorldPartitionActorDesc> ActorDescPtr;
-};
-
-class FEditorGrid2DDirtyActorDescInstance : public FEditorGrid2DActorDescInstance
-{
-public:
-	FEditorGrid2DDirtyActorDescInstance(const IWorldPartitionActorDescInstance* InActorDescInstance)
-		: ActorDescInstance(InActorDescInstance)
-	{
-	}
-
-	virtual ~FEditorGrid2DDirtyActorDescInstance() {}
-		
-	virtual const FWorldPartitionActorDesc* GetActorDesc() const override { return ActorDescInstance->GetActorDesc(); }
-
-protected:
-	
-	const IWorldPartitionActorDescInstance* ActorDescInstance;
-};
-
 class FWeightedMovingAverageScope
 {
 public:
@@ -347,7 +280,6 @@ void SWorldPartitionEditorGrid2D::SToolBar::Construct(const FArguments& InArgs)
 		return ToolbarBuilder.MakeWidget();
 	};
 
-
 	const FMargin ToolbarSlotPadding(4.0f, 1.0f);
 	const FMargin ToolbarButtonPadding(4.0f, 0.0f);
 
@@ -517,12 +449,7 @@ SWorldPartitionEditorGrid2D::SWorldPartitionEditorGrid2D()
 }
 
 SWorldPartitionEditorGrid2D::~SWorldPartitionEditorGrid2D()
-{
-	if (GEngine)
-	{
-		GEngine->OnLevelActorAdded().RemoveAll(this);
-	}
-}
+{}
 
 void SWorldPartitionEditorGrid2D::Construct(const FArguments& InArgs)
 {
@@ -601,7 +528,11 @@ void SWorldPartitionEditorGrid2D::Construct(const FArguments& InArgs)
 
 	BindCommands();
 
-	GEngine->OnLevelActorAdded().AddRaw(this, &SWorldPartitionEditorGrid2D::OnActorAdded);
+	ExternalDirtyActorsTracker.Reset();
+	if (InArgs._InWorld)
+	{
+		ExternalDirtyActorsTracker = MakeUnique<FExternalDirtyActorsTracker>(InArgs._InWorld->PersistentLevel, this);
+	}
 }
 
 void SWorldPartitionEditorGrid2D::BindCommands()
@@ -1010,14 +941,6 @@ TSharedRef<SWidget> SWorldPartitionEditorGrid2D::GenerateContextualMenu() const
 	return UToolMenus::Get()->GenerateWidget(MenuName, FToolMenuContext(CommandList));
 }
 
-void SWorldPartitionEditorGrid2D::OnActorAdded(AActor* Actor)
-{
-	if (!Actor->GetWorld()->IsGameWorld() && Actor->IsPackageExternal() && !Actor->IsChildActor())
-	{
-		NewlyAddedUnsavedActorDescs.OnActorAdded(Actor);
-	}
-}
-
 FReply SWorldPartitionEditorGrid2D::OnMouseButtonUp(const FGeometry& MyGeometry, const FPointerEvent& MouseEvent)
 {
 	const bool bIsLeftMouseButtonEffecting = MouseEvent.GetEffectingButton() == EKeys::LeftMouseButton;
@@ -1317,39 +1240,31 @@ void SWorldPartitionEditorGrid2D::Tick(const FGeometry& AllottedGeometry, const
 		}
 	}, ForEachIntersectingActorParams);
 
-	// Add dirty actors as an acceleration for FDirtyActorDescView constructor
-	for (auto& [Reference, Actor] : GetWorldPartition()->GetDirtyActors())
+	// Also include dirty actors so we can display updated bounds.
+	if (ExternalDirtyActorsTracker.IsValid())
 	{
-		DirtyActorGuids.Add(Actor->GetActorGuid());
-	}
-
-	// Also include transient actor loader adapters that might have been spawned by blutilities, etc. Since these actors can't be saved because they are transient,
-	// they will never get an actor descriptor so they will never appear in the world partition editor. Also include unsaved, newly created actors for convenience.
-	for (FActorDescInstanceList::TIterator<> ActorDescIterator(&NewlyAddedUnsavedActorDescs); ActorDescIterator; ++ActorDescIterator)
-	{
-		if (AActor* NewlyAddedUnsavedActor = ActorDescIterator->GetActor())
+		for (auto& [WeakActor, ActorGuid] : ExternalDirtyActorsTracker->GetDirtyActors())
 		{
-			if (!NewlyAddedUnsavedActor->GetPackage()->IsDirty())
+			if (WeakActor.IsValid())
 			{
-				ActorDescIterator.RemoveCurrent();
-			}
-			else
-			{
-				if (NewlyAddedUnsavedActor->Implements<UWorldPartitionActorLoaderInterface>())
+				DirtyActorGuids.Add(ActorGuid);
+
+				// Also include transient actor loader adapters that might have been spawned by blutilities, etc. Since these actors can't be saved because they are transient,
+				// they will never get an actor descriptor so they will never appear in the world partition editor. Also include unsaved, newly created actors for convenience.
+				if (WeakActor->Implements<UWorldPartitionActorLoaderInterface>())
 				{
-					if (IWorldPartitionActorLoaderInterface::ILoaderAdapter* LoaderAdapter = Cast<IWorldPartitionActorLoaderInterface>(NewlyAddedUnsavedActor)->GetLoaderAdapter())
+					if (IWorldPartitionActorLoaderInterface::ILoaderAdapter* LoaderAdapter = Cast<IWorldPartitionActorLoaderInterface>(WeakActor.Get())->GetLoaderAdapter())
 					{
-						ShownLoaderInterfaces.Add(NewlyAddedUnsavedActor);
+						ShownLoaderInterfaces.Add(WeakActor.Get());
 					}
 				}
 
-				ShownActorGuids.Add(NewlyAddedUnsavedActor->GetActorGuid());
+				if (bShowActors)
+				{
+					ShownActorGuids.Add(ActorGuid);
+				}
 			}
 		}
-		else if (!ActorDescIterator->GetActor(true))
-		{
-			ActorDescIterator.RemoveCurrent();
-		}
 	}
 
 	FLoaderInterfaceSet LastHoveredLoaderInterfaces = MoveTemp(HoveredLoaderInterfaces);
@@ -1427,46 +1342,61 @@ void SWorldPartitionEditorGrid2D::Tick(const FGeometry& AllottedGeometry, const
 	for (const FWorldPartitionHandle& SelectedActorHandle : WorldPartitionSubsystem->SelectedActorHandles)
 	{
 		ShownActorGuids.Add(SelectedActorHandle.GetInstance()->GetGuid());
+		SelectedActorGuids.Add(SelectedActorHandle.GetInstance()->GetGuid());
 	}
 }
 
 uint32 SWorldPartitionEditorGrid2D::PaintActors(const FGeometry& AllottedGeometry, const FSlateRect& MyCullingRect, FSlateWindowElementList& OutDrawElements, uint32 LayerId) const
 {
+	const UWorldPartition* ThisWorldPartition = GetWorldPartition();
 	const FBox2D ViewRect(FVector2D(ForceInitToZero), AllottedGeometry.GetLocalSize());
 	const FBox2D WorldViewRect(ScreenToWorld.TransformPoint(ViewRect.Min), ScreenToWorld.TransformPoint(ViewRect.Max));
 	const FBox ViewRectWorld(FVector(WorldViewRect.Min.X, WorldViewRect.Min.Y, -HALF_WORLD_MAX), FVector(WorldViewRect.Max.X, WorldViewRect.Max.Y, HALF_WORLD_MAX));
 	const UWorldPartitionSubsystem* WorldPartitionSubsystem = UWorld::GetSubsystem<UWorldPartitionSubsystem>(GetWorld());
 
-	// Reserve here is done to avoid reallocations so that pointers in the BoundProxies set stay valid
-	TArray<FEditorGrid2DDirtyActorDescInstance> DirtyActorDescInstances;
-	DirtyActorDescInstances.Reserve(DirtyActorGuids.Num());
-	
-	TSet<const IWorldPartitionActorDescInstance*> BoundProxies;
-	BoundProxies.Reserve(ShownActorGuids.Num());
-
 	const IWorldPartitionActorLoaderInterface::ILoaderAdapter* LocalHoveredLoaderAdapter = nullptr;
 	if (IWorldPartitionActorLoaderInterface* HoveredLoaderAdapterInterface = Cast<IWorldPartitionActorLoaderInterface>(HoveredLoaderInterface.Get()))
 	{
 		LocalHoveredLoaderAdapter = HoveredLoaderInterface->GetLoaderAdapter();
 	}
 
-	const UWorldPartition* ThisWorldPartition = GetWorldPartition();
+	struct FActorBoundsDesc
+	{
+		FActorBoundsDesc(const FWorldPartitionActorDescInstance* InActorDescInstance, const AActor* InActor)
+		{
+			Guid = InActor ? InActor->GetActorGuid() : InActorDescInstance->GetGuid();
+			Label = InActor ? *InActor->GetActorLabel(false) : InActorDescInstance->GetActorLabel();
+			DescBounds = InActorDescInstance->GetEditorBounds();
+			ActorBounds = InActor ? InActor->GetStreamingBounds() : DescBounds;
+			bIsSpatiallyLoaded = InActor ? InActor->GetIsSpatiallyLoaded() : InActorDescInstance->GetIsSpatiallyLoaded();
+		}
+
+		FGuid Guid;
+		FName Label;
+		FBox DescBounds;
+		FBox ActorBounds;		
+		bool bIsSpatiallyLoaded;
+	};
+
+	TArray<FActorBoundsDesc> ActorBoundsDescs;
+	ActorBoundsDescs.Reserve(ShownActorGuids.Num());
+	
 	for (const FGuid& ActorGuid : ShownActorGuids)
 	{
 		if (const FWorldPartitionActorDescInstance* ActorDescInstance = ThisWorldPartition->GetActorDescInstance(ActorGuid))
 		{
-			if (DirtyActorGuids.Contains(ActorGuid))
-			{
-				BoundProxies.Add(&DirtyActorDescInstances.Emplace_GetRef(ActorDescInstance));
-			}
-			else
-			{
-				BoundProxies.Add(ActorDescInstance);
-			}
+			ActorBoundsDescs.Emplace(ActorDescInstance, DirtyActorGuids.Contains(ActorGuid) ? ActorDescInstance->GetActor(false) : nullptr);
 		}
-		else if (const FWorldPartitionActorDescInstance* NewlyAddedUnsavedActorDesc = NewlyAddedUnsavedActorDescs.GetActorDescInstance(ActorGuid))
+	}
+
+	if (ExternalDirtyActorsTracker.IsValid())
+	{
+		for (auto& [WeakActor, ActorGuid] : ExternalDirtyActorsTracker->GetDirtyActors())
 		{
-			BoundProxies.Add(NewlyAddedUnsavedActorDesc);
+			if (WeakActor.IsValid() && (bShowActors || SelectedActorGuids.Contains(ActorGuid)) && !ThisWorldPartition->GetActorDescInstance(ActorGuid))
+			{
+				ActorBoundsDescs.Emplace(nullptr, WeakActor.Get());
+			}
 		}
 	}
 
@@ -1648,13 +1578,13 @@ uint32 SWorldPartitionEditorGrid2D::PaintActors(const FGeometry& AllottedGeometr
 		}
 	}
 
-	if (BoundProxies.Num())
+	if (ActorBoundsDescs.Num())
 	{
 		TArray<FVector2D> LinePoints;
 		LinePoints.SetNum(5);
 
 		bool bIsBoundsEdgeHoveredFound = false;
-		for (const IWorldPartitionActorDescInstance* BoundProxy : BoundProxies)
+		for (const FActorBoundsDesc& ActorBoundsDesc : ActorBoundsDescs)
 		{
 			auto ShowActorBox = [this, &AllottedGeometry, &OutDrawElements, &LayerId, &bIsBoundsEdgeHoveredFound, &DrawActorLabel](const FBox& ActorBounds, bool bIsSelected, bool bIsSpatiallyLoaded, const FName ActorLabel)
 			{
@@ -1723,15 +1653,14 @@ uint32 SWorldPartitionEditorGrid2D::PaintActors(const FGeometry& AllottedGeometr
 				}
 			};
 			
-			const FWorldPartitionHandle ActorHandle(BoundProxy->GetContainerInstance(), BoundProxy->GetGuid());
-			const FBox ActorBounds = BoundProxy->GetEditorBounds();
-			const bool bIsSelected = SelectedActorGuids.Contains(BoundProxy->GetGuid()) || (ActorHandle.IsValid() && WorldPartitionSubsystem->SelectedActorHandles.Contains(ActorHandle));
-			const bool bIsSpatiallyLoaded = BoundProxy->GetIsSpatiallyLoaded();
-			const FName ActorLabel = BoundProxy->GetActorLabel();
+			const FBox ActorBounds = ActorBoundsDesc.ActorBounds;
+			const bool bIsSelected = SelectedActorGuids.Contains(ActorBoundsDesc.Guid);
+			const bool bIsSpatiallyLoaded = ActorBoundsDesc.bIsSpatiallyLoaded;
+			const FName ActorLabel = ActorBoundsDesc.Label;
 
 			if (bIsSelected)
 			{
-				const FBox ActorDescBounds = BoundProxy->GetActorDesc()->GetEditorBounds();
+				const FBox ActorDescBounds = ActorBoundsDesc.DescBounds;
 				if (!ActorDescBounds.Equals(ActorBounds, 1.0f))
 				{
 					ShowActorBox(ActorDescBounds, false, bIsSpatiallyLoaded, ActorLabel);
@@ -2371,28 +2300,4 @@ void SWorldPartitionEditorGrid2D::ClearSelection()
 	SelectBoxGridSnapped.Init();
 }
 
-void SWorldPartitionEditorGrid2D::FNewlyAddedUnsavedActorDescsDescRegistry::OnActorAdded(AActor* Actor)
-{
-	if (TUniquePtr<FWorldPartitionActorDescInstance>* ExistingActorDescInstance = GetActorDescriptor(Actor->GetActorGuid()))
-	{
-		FEditorGrid2DUnsavedActorDescInstance* NewlyAddedActorDescInstance = StaticCast<FEditorGrid2DUnsavedActorDescInstance*>(ExistingActorDescInstance->Get());
-
-		FWorldPartitionActorDescUtils::UpdateActorDescriptorFromActor(Actor, NewlyAddedActorDescInstance->ActorDescPtr);
-	}
-	else
-	{
-		AddActorDescriptor(new FEditorGrid2DUnsavedActorDescInstance(Actor));
-	}
-}
-
-const FWorldPartitionActorDescInstance* SWorldPartitionEditorGrid2D::FNewlyAddedUnsavedActorDescsDescRegistry::GetActorDescInstance(const FGuid& InActorGuid) const
-{
-	if (const TUniquePtr<FWorldPartitionActorDescInstance>* ExistingActorDescInstance = GetActorDescriptor(InActorGuid))
-	{
-		return ExistingActorDescInstance->Get();
-	}
-
-	return nullptr;
-}
-
 #undef LOCTEXT_NAMESPACE
\ No newline at end of file
diff --git a/Engine/Source/Editor/WorldPartitionEditor/Public/WorldPartition/SWorldPartitionEditorGrid2D.h b/Engine/Source/Editor/WorldPartitionEditor/Public/WorldPartition/SWorldPartitionEditorGrid2D.h
--- a/Engine/Source/Editor/WorldPartitionEditor/Public/WorldPartition/SWorldPartitionEditorGrid2D.h
+++ b/Engine/Source/Editor/WorldPartitionEditor/Public/WorldPartition/SWorldPartitionEditorGrid2D.h
@@ -11,8 +11,10 @@
 #include "SWorldPartitionViewportWidget.h"
 #include "Widgets/Text/STextBlock.h"
 #include "SViewportToolBar.h"
+#include "WorldPartition/WorldPartition.h"
 #include "WorldPartition/ActorDescContainerInstance.h"
 #include "WorldPartition/WorldPartitionActorLoaderInterface.h"
+#include "ExternalDirtyActorsTracker.h"
 
 class SWorldPartitionEditorGrid2D : public SWorldPartitionEditorGrid
 {
@@ -49,14 +51,12 @@ protected:
 		TSharedPtr<FUICommandInfo> FocusLoadedRegions;
 		TSharedPtr<FUICommandInfo> FocusWorld;
 		
-
 		/**
 		 * Initialize commands
 		 */
 		virtual void RegisterCommands() override;
 	};
 
-
 	// In-viewport toolbar widget used in the world partition editor
 	class SToolBar : public SViewportToolBar
 	{
@@ -229,14 +229,13 @@ protected:
 	TSharedPtr<STextBlock> TextWorldBoundsInKMWidget;
 	TSharedPtr<STextBlock> TextRulerWidget;
 
-	// List of actors without an actor descriptor
-	class FNewlyAddedUnsavedActorDescsDescRegistry : public FActorDescInstanceList
+	struct FExternalDirtyActorTrackerGuid
 	{
-	public:
-		void OnActorAdded(AActor* Actor);
-		
-		const FWorldPartitionActorDescInstance* GetActorDescInstance(const FGuid& InActorGuid) const;
+		using Type = FGuid;
+		using OwnerType = SWorldPartitionEditorGrid2D;
+		static FGuid Store(SWorldPartitionEditorGrid2D* InOwner, AActor* InActor) { return InActor->GetActorGuid(); }
 	};
 
-	FNewlyAddedUnsavedActorDescsDescRegistry NewlyAddedUnsavedActorDescs;
+	using FExternalDirtyActorsTracker = TExternalDirtyActorsTracker<FExternalDirtyActorTrackerGuid>;
+	TUniquePtr<FExternalDirtyActorsTracker> ExternalDirtyActorsTracker;
 };
\ No newline at end of file
diff --git a/Engine/Source/Runtime/Engine/Private/WorldPartition/WorldPartition.cpp b/Engine/Source/Runtime/Engine/Private/WorldPartition/WorldPartition.cpp
--- a/Engine/Source/Runtime/Engine/Private/WorldPartition/WorldPartition.cpp
+++ b/Engine/Source/Runtime/Engine/Private/WorldPartition/WorldPartition.cpp
@@ -251,6 +251,42 @@ public:
 	}
 };
 
+UWorldPartition::FWorldPartitionExternalDirtyActorsTracker::FWorldPartitionExternalDirtyActorsTracker()
+	: Super(nullptr, nullptr)
+{}
+
+UWorldPartition::FWorldPartitionExternalDirtyActorsTracker::FWorldPartitionExternalDirtyActorsTracker(UWorldPartition* InWorldPartition)
+	: Super(InWorldPartition->GetTypedOuter<ULevel>(), InWorldPartition)
+{}
+
+void UWorldPartition::FWorldPartitionExternalDirtyActorsTracker::OnRemoveNonDirtyActor(TWeakObjectPtr<AActor> InActor, FWorldPartitionReference& InValue)
+{
+	check(InActor.IsValid());
+	NonDirtyActors.Emplace(InActor, InValue);
+}
+
+void UWorldPartition::FWorldPartitionExternalDirtyActorsTracker::Tick(float InDeltaSeconds)
+{
+	Super::Tick(InDeltaSeconds);
+
+	for (auto& [Actor, Reference] : NonDirtyActors)
+	{
+		// Resolve reference for newly added actors
+		if (!Reference.IsValid())
+		{
+			Reference = FWorldPartitionReference(Owner, Actor->GetActorGuid());
+		}
+
+		// Transfer ownership of our last ref if actor can be pinned
+		if (Reference.IsValid() && Reference.GetInstance()->GetHardRefCount() <= 1 && Owner->PinnedActors && FLoaderAdapterPinnedActors::SupportsPinning(Reference.GetInstance()))
+		{
+			Owner->PinnedActors->AddActors({ Reference.ToHandle() });
+		}
+	}
+
+	NonDirtyActors.Empty();
+}
+
 UWorldPartition::FWorldPartitionChangedEvent UWorldPartition::WorldPartitionChangedEvent;
 #endif
 
@@ -348,54 +384,6 @@ void UWorldPartition::OnGCPostReachabilityAnalysis()
 	}
 }
 
-void UWorldPartition::OnObjectsReplaced(const TMap<UObject*, UObject*>& OldToNewObjectMap)
-{
-	// Patch up actor pointers in dirty actor list
-	for (auto [OldObject, NewObject] : OldToNewObjectMap)
-	{
-		if (AActor* OldActor = Cast<AActor>(OldObject))
-		{
-			if (AActor* NewActor = Cast<AActor>(NewObject))
-			{
-				for (auto& [ActorReference, Actor] : DirtyActors)
-				{
-					if (Actor == OldActor)
-					{
-						Actor = NewActor;
-					}
-				}
-			}
-		}
-	}
-}
-
-void UWorldPartition::OnPackageDirtyStateChanged(UPackage* Package)
-{	
-	auto ShouldHandleActor = [this](AActor* Actor)
-	{
-		if (Actor && Actor->IsMainPackageActor())
-		{
-			if (ULevel* OuterLevel = Actor->GetTypedOuter<ULevel>())
-			{
-				return OuterLevel->GetWorldPartition() == this;
-			}
-		}
-		return false;
-	};
-
-	if (AActor* Actor = AActor::FindActorInPackage(Package); ShouldHandleActor(Actor))
-	{
-		if (FWorldPartitionHandle ActorHandle(this, Actor->GetActorGuid()); ActorHandle.IsValid())
-		{
-			if (Package->IsDirty())
-			{
-				DirtyActors.Add(ActorHandle.ToReference(), Actor);
-			}
-		}
-	}
-
-}
-
 // Returns whether the memory package is part of the known/valid package names
 // used by World Partition for PIE/-game streaming.
 bool UWorldPartition::IsValidPackageName(const FString& InPackageName)
@@ -768,8 +756,6 @@ void UWorldPartition::Uninitialize()
 			RegisteredEditorLoaderAdapters.Empty();
 		}
 
-		DirtyActors.Empty();
-
 		UninitializeActorDescContainers();
 		ActorDescContainerInstance = nullptr;
 
@@ -878,7 +864,6 @@ void UWorldPartition::RegisterDelegates()
 			FEditorDelegates::CancelPIE.AddUObject(this, &UWorldPartition::OnCancelPIE);
 			FGameDelegates::Get().GetEndPlayMapDelegate().AddUObject(this, &UWorldPartition::OnEndPlay);
 			FCoreUObjectDelegates::PostReachabilityAnalysis.AddUObject(this, &UWorldPartition::OnGCPostReachabilityAnalysis);
-			FCoreUObjectDelegates::OnObjectsReplaced.AddUObject(this, &UWorldPartition::OnObjectsReplaced);
 			GEditor->OnPostBugItGoCalled().AddUObject(this, &UWorldPartition::OnPostBugItGoCalled);
 			GEditor->OnEditorClose().AddUObject(this, &UWorldPartition::SavePerUserSettings);
 			FWorldDelegates::OnPostWorldRename.AddUObject(this, &UWorldPartition::OnWorldRenamed);
@@ -886,7 +871,7 @@ void UWorldPartition::RegisterDelegates()
 
 		if (!IsRunningCommandlet())
 		{
-			UPackage::PackageDirtyStateChangedEvent.AddUObject(this, &UWorldPartition::OnPackageDirtyStateChanged);
+			ExternalDirtyActorsTracker = MakeUnique<FWorldPartitionExternalDirtyActorsTracker>(this);
 		}
 	}
 #endif
@@ -922,7 +907,6 @@ void UWorldPartition::UnregisterDelegates()
 			if (!IsEngineExitRequested())
 			{
 				FCoreUObjectDelegates::PostReachabilityAnalysis.RemoveAll(this);
-				FCoreUObjectDelegates::OnObjectsReplaced.RemoveAll(this);
 			}
 
 			GEditor->OnPostBugItGoCalled().RemoveAll(this);
@@ -931,7 +915,7 @@ void UWorldPartition::UnregisterDelegates()
 
 		if (!IsRunningCommandlet())
 		{
-			UPackage::PackageDirtyStateChangedEvent.RemoveAll(this);
+			ExternalDirtyActorsTracker.Reset();
 		}
 	}
 #endif
@@ -1214,11 +1198,6 @@ void UWorldPartition::OnActorDescInstanceAdded(FWorldPartitionActorDescInstance*
 
 	HashActorDescInstance(NewActorDescInstance);
 
-	if (AActor* NewActor = NewActorDescInstance->GetActor())
-	{
-		DirtyActors.Add(FWorldPartitionReference(NewActorDescInstance->GetContainerInstance(), NewActorDescInstance->GetGuid()), NewActor);
-	}
-
 	if (ForceLoadedActors)
 	{
 		ForceLoadedActors->AddActors({ NewActorDescInstance->GetGuid() });
@@ -1513,12 +1492,15 @@ void UWorldPartition::AddReferencedObjects(UObject* InThis, FReferenceCollector&
 	// by the transaction buffer, but we clear it when unloading regions, etc. and we don't want these actors to die.
 	// Also, we must avoid reporting these references when not collecting garbage, as code such as package deletion
 	// will skip packages with actors still referenced (via GatherObjectReferencersForDeletion).
-	if (IsGarbageCollecting())
+	if (This->ExternalDirtyActorsTracker.IsValid() && IsGarbageCollecting())
 	{
 		Collector.AllowEliminatingReferences(false);
-		for (auto& [ActorReference, Actor] : This->DirtyActors)
+		for (auto& [WeakActor, Value] : This->ExternalDirtyActorsTracker->GetDirtyActors())
 		{
-			Collector.AddReferencedObject(Actor);
+			if (TObjectPtr<AActor> Actor = WeakActor.Get(true))
+			{
+				Collector.AddReferencedObject(Actor);
+			}
 		}
 		Collector.AllowEliminatingReferences(true);
 	}
@@ -1540,25 +1522,9 @@ void UWorldPartition::Tick(float DeltaSeconds)
 		EditorHash->Tick(DeltaSeconds);
 	}
 
-	for (TMap<FWorldPartitionReference, TObjectPtr<AActor>>::TIterator DirtyActorIt(DirtyActors); DirtyActorIt; ++DirtyActorIt)
+	if (ExternalDirtyActorsTracker)
 	{
-		if (!DirtyActorIt.Key().IsValid())
-		{
-			DirtyActorIt.RemoveCurrent();
-		}
-		else if (!DirtyActorIt.Value()->GetPackage()->IsDirty())
-		{
-			// Transfer ownership of the last ref if actor can be pinned
-			if (DirtyActorIt.Key().GetInstance()->GetHardRefCount() <= 1 && PinnedActors && FLoaderAdapterPinnedActors::SupportsPinning(DirtyActorIt.Key().GetInstance()))
-			{
-				PinnedActors->AddActors({ DirtyActorIt.Key().ToHandle() });
-				DirtyActorIt.RemoveCurrent();
-			} // Clean up if a loader took a reference on the actor (ILoaderAdapter::RefreshLoadedState was called since)
-			else if (DirtyActorIt.Key().GetInstance()->GetHardRefCount() > 1)
-			{
-				DirtyActorIt.RemoveCurrent();
-			}
-		}
+		ExternalDirtyActorsTracker->Tick(DeltaSeconds);
 	}
 
 	if (bForceGarbageCollection)
@@ -2034,14 +2000,6 @@ bool UWorldPartition::UnregisterActorDescContainerInstance(UActorDescContainerIn
 			if (ActorHandle.IsValid())
 			{
 				ActorGuids.Add(It->GetGuid());
-
-				for (TMap<FWorldPartitionReference, TObjectPtr<AActor>>::TIterator DirtyActorIt(DirtyActors); DirtyActorIt; ++DirtyActorIt)
-				{
-					if (DirtyActorIt.Key() == ActorHandle)
-					{
-						DirtyActorIt.RemoveCurrent();
-					}
-				}
 			}
 		}
 
diff --git a/Engine/Source/Runtime/Engine/Public/ExternalDirtyActorsTracker.h b/Engine/Source/Runtime/Engine/Public/ExternalDirtyActorsTracker.h
new file mode 100644
--- /dev/null
+++ b/Engine/Source/Runtime/Engine/Public/ExternalDirtyActorsTracker.h
@@ -0,0 +1,122 @@
+// Copyright Epic Games, Inc. All Rights Reserved.
+
+#pragma once
+
+#if WITH_EDITOR
+#include "TickableEditorObject.h"
+#include "UObject/Package.h"
+
+/**
+ * TExternalDirtyActorsTracker is a tracker for dirty external actors, with custom storage through the StoreType interface.
+ */
+template <typename StoreType>
+class TExternalDirtyActorsTracker : public FTickableEditorObject
+{
+public:
+	using Super = TExternalDirtyActorsTracker<StoreType>;
+	using MapType = TMap<TWeakObjectPtr<AActor>, typename StoreType::Type>;
+
+	TExternalDirtyActorsTracker(const ULevel* InLevel, typename StoreType::OwnerType* InOwner)
+		: Level(InLevel)
+		, Owner(InOwner)
+	{
+		UPackage::PackageDirtyStateChangedEvent.AddRaw(this, &TExternalDirtyActorsTracker::OnPackageDirtyStateChanged);
+		FCoreUObjectDelegates::OnObjectsReplaced.AddRaw(this, &TExternalDirtyActorsTracker::OnObjectsReplaced);
+	}
+
+	~TExternalDirtyActorsTracker()
+	{
+		UPackage::PackageDirtyStateChangedEvent.RemoveAll(this);
+		FCoreUObjectDelegates::OnObjectsReplaced.RemoveAll(this);
+	}
+
+	const MapType& GetDirtyActors() const { return DirtyActors; }
+
+	virtual void OnRemoveInvalidDirtyActor(const TWeakObjectPtr<AActor> InActor, StoreType::Type& InValue) {}
+	virtual void OnRemoveNonDirtyActor(const TWeakObjectPtr<AActor> InActor, StoreType::Type& InValue) {}
+
+protected:
+	void OnPackageDirtyStateChanged(UPackage* InPackage)
+	{
+		if (AActor* Actor = AActor::FindActorInPackage(InPackage))
+		{
+			if (ULevel* OuterLevel = Actor->GetTypedOuter<ULevel>())
+			{
+				if (OuterLevel == Level)
+				{
+					if (InPackage->IsDirty())
+					{
+						DirtyActors.Add(Actor, StoreType::Store(Owner, Actor));
+					}
+					else
+					{
+						typename StoreType::Type Value;
+						if (DirtyActors.RemoveAndCopyValue(Actor, Value))
+						{
+							if (IsValid(Actor))
+							{
+								OnRemoveNonDirtyActor(Actor, Value);
+							}
+							else
+							{
+								OnRemoveInvalidDirtyActor(Actor, Value);
+							}
+						}
+					}
+				}
+			}
+		}
+	}
+
+	void OnObjectsReplaced(const TMap<UObject*, UObject*>& OldToNewObjectMap)
+	{
+		for (auto [OldObject, NewObject] : OldToNewObjectMap)
+		{
+			if (AActor* OldActor = Cast<AActor>(OldObject))
+			{
+				if (AActor* NewActor = Cast<AActor>(NewObject))
+				{
+					for (auto& [WeakActor, Value] : DirtyActors)
+					{
+						if (WeakActor.IsValid() && WeakActor.Get() == OldActor)
+						{
+							WeakActor = NewActor;
+						}
+					}
+				}
+			}
+		}
+	}
+
+	//~ Begin FTickableEditorObject interface
+	virtual TStatId GetStatId() const override
+	{
+		return TStatId();
+	}
+
+	virtual void Tick(float DeltaTime) override
+	{
+		for (typename MapType::TIterator DirtyActorIt(DirtyActors); DirtyActorIt; ++DirtyActorIt)
+		{
+			if (DirtyActorIt.Key().IsValid())
+			{
+				if (!DirtyActorIt.Key().Get()->GetPackage()->IsDirty())
+				{
+					OnRemoveNonDirtyActor(DirtyActorIt.Key(), DirtyActorIt.Value());
+					DirtyActorIt.RemoveCurrent();
+				}
+			}
+			else if (!DirtyActorIt.Key().IsValid(true))
+			{
+				OnRemoveInvalidDirtyActor(DirtyActorIt.Key(), DirtyActorIt.Value());
+				DirtyActorIt.RemoveCurrent();
+			}
+		}
+	}
+	//~ End FTickableEditorObject interface
+
+	const ULevel* Level;
+	typename StoreType::OwnerType* Owner;
+	MapType DirtyActors;
+};
+#endif
\ No newline at end of file
diff --git a/Engine/Source/Runtime/Engine/Public/WorldPartition/WorldPartition.h b/Engine/Source/Runtime/Engine/Public/WorldPartition/WorldPartition.h
--- a/Engine/Source/Runtime/Engine/Public/WorldPartition/WorldPartition.h
+++ b/Engine/Source/Runtime/Engine/Public/WorldPartition/WorldPartition.h
@@ -22,6 +22,7 @@
 #include "WorldPartition/WorldPartitionStreamingGeneration.h"
 #include "WorldPartition/WorldPartitionActorLoaderInterface.h"
 #include "WorldPartition/WorldPartitionEditorLoaderAdapter.h"
+#include "ExternalDirtyActorsTracker.h"
 #include "PackageSourceControlHelper.h"
 #include "CookPackageSplitter.h"
 #include "Delegates/DelegateCombinations.h"
@@ -148,7 +149,6 @@ private:
 	ENGINE_API void SavePerUserSettings();
 		
 	ENGINE_API void OnGCPostReachabilityAnalysis();
-	ENGINE_API void OnObjectsReplaced(const TMap<UObject*, UObject*>& OldToNewObjectMap);
 	ENGINE_API void OnPackageDirtyStateChanged(UPackage* Package);
 
 	// PIE/Game
@@ -406,8 +406,6 @@ public:
 	bool IsEnablingStreamingJustified() const { return bEnablingStreamingJustified; }
 
 	bool IsHLODsInEditorAllowed() const { return bAllowShowingHLODsInEditor; }
-
-	const TMap<FWorldPartitionReference, AActor*>& GetDirtyActors() const { return ObjectPtrDecay(DirtyActors); }
 #endif
 
 public:
@@ -546,7 +544,29 @@ private:
 
 	TArray<FWorldPartitionReference> LoadedSubobjects;
 
-	TMap<FWorldPartitionReference, TObjectPtr<AActor>> DirtyActors;
+	struct FWorldPartitionExternalDirtyActorsTrackerReference
+	{
+		using Type = FWorldPartitionReference;
+		using OwnerType = UWorldPartition;
+		static FWorldPartitionReference Store(UWorldPartition* InOwner, AActor* InActor) { return FWorldPartitionReference(InOwner, InActor->GetActorGuid()); }
+	};
+
+	class FWorldPartitionExternalDirtyActorsTracker : public TExternalDirtyActorsTracker<FWorldPartitionExternalDirtyActorsTrackerReference>
+	{
+	public:
+		FWorldPartitionExternalDirtyActorsTracker();
+		FWorldPartitionExternalDirtyActorsTracker(UWorldPartition* InWorldPartition);
+
+		//~ Begin TExternalDirtyActorsTracker interface
+		virtual void OnRemoveNonDirtyActor(const TWeakObjectPtr<AActor> InActor, FWorldPartitionReference& InValue) override;
+		virtual void Tick(float InDeltaTime) override;
+		//~ End TExternalDirtyActorsTracker interface
+
+	private:
+		TArray<TPair<TWeakObjectPtr<AActor>, FWorldPartitionReference>> NonDirtyActors;
+	};
+
+	TUniquePtr<FWorldPartitionExternalDirtyActorsTracker> ExternalDirtyActorsTracker;
 
 	TSet<FString> GeneratedStreamingPackageNames;
 
