World Partition: fix crash when placing a Level Instance.

diff --git a/Engine/Source/Editor/WorldPartitionEditor/Private/WorldPartition/SWorldPartitionEditorGrid2D.cpp b/Engine/Source/Editor/WorldPartitionEditor/Private/WorldPartition/SWorldPartitionEditorGrid2D.cpp
--- a/Engine/Source/Editor/WorldPartitionEditor/Private/WorldPartition/SWorldPartitionEditorGrid2D.cpp
+++ b/Engine/Source/Editor/WorldPartitionEditor/Private/WorldPartition/SWorldPartitionEditorGrid2D.cpp
@@ -1364,11 +1364,21 @@ uint32 SWorldPartitionEditorGrid2D::PaintActors(const FGeometry& AllottedGeometr
 	{
 		FActorBoundsDesc(const FWorldPartitionActorDescInstance* InActorDescInstance, const AActor* InActor)
 		{
-			Guid = InActor ? InActor->GetActorGuid() : InActorDescInstance->GetGuid();
-			Label = InActor ? *InActor->GetActorLabel(false) : InActorDescInstance->GetActorLabel();
-			DescBounds = InActorDescInstance->GetEditorBounds();
-			ActorBounds = InActor ? InActor->GetStreamingBounds() : DescBounds;
-			bIsSpatiallyLoaded = InActor ? InActor->GetIsSpatiallyLoaded() : InActorDescInstance->GetIsSpatiallyLoaded();
+			if (InActor)
+			{
+				Guid = InActor->GetActorGuid();
+				Label = *InActor->GetActorLabel(false);
+				ActorBounds = InActor->GetStreamingBounds();
+				DescBounds = InActorDescInstance ? InActorDescInstance->GetEditorBounds() : ActorBounds;
+				bIsSpatiallyLoaded = InActor->GetIsSpatiallyLoaded();
+			}
+			else if (InActorDescInstance)
+			{
+				Guid = InActorDescInstance->GetGuid();
+				Label = InActorDescInstance->GetActorLabel();
+				ActorBounds = DescBounds = InActorDescInstance->GetEditorBounds();
+				bIsSpatiallyLoaded = InActorDescInstance->GetIsSpatiallyLoaded();
+			}
 		}
 
 		FGuid Guid;
@@ -1652,22 +1662,25 @@ uint32 SWorldPartitionEditorGrid2D::PaintActors(const FGeometry& AllottedGeometr
 					);
 				}
 			};
-			
-			const FBox ActorBounds = ActorBoundsDesc.ActorBounds;
-			const bool bIsSelected = SelectedActorGuids.Contains(ActorBoundsDesc.Guid);
-			const bool bIsSpatiallyLoaded = ActorBoundsDesc.bIsSpatiallyLoaded;
-			const FName ActorLabel = ActorBoundsDesc.Label;
 
-			if (bIsSelected)
-			{
-				const FBox ActorDescBounds = ActorBoundsDesc.DescBounds;
-				if (!ActorDescBounds.Equals(ActorBounds, 1.0f))
+			if (ActorBoundsDesc.Guid.IsValid())
+			{			
+				const FBox ActorBounds = ActorBoundsDesc.ActorBounds;
+				const bool bIsSelected = SelectedActorGuids.Contains(ActorBoundsDesc.Guid);
+				const bool bIsSpatiallyLoaded = ActorBoundsDesc.bIsSpatiallyLoaded;
+				const FName ActorLabel = ActorBoundsDesc.Label;
+
+				if (bIsSelected)
 				{
-					ShowActorBox(ActorDescBounds, false, bIsSpatiallyLoaded, ActorLabel);
+					const FBox ActorDescBounds = ActorBoundsDesc.DescBounds;
+					if (!ActorDescBounds.Equals(ActorBounds, 1.0f))
+					{
+						ShowActorBox(ActorDescBounds, false, bIsSpatiallyLoaded, ActorLabel);
+					}
 				}
-			}
 
-			ShowActorBox(ActorBounds, bIsSelected, bIsSpatiallyLoaded, ActorLabel);
+				ShowActorBox(ActorBounds, bIsSelected, bIsSpatiallyLoaded, ActorLabel);
+			}
 		};
 	}
 
