World Partition Editor Optimization: cache intersecting viewport actors, saves a constant ~10-12ms when the world partition editor is open.

diff --git a/Engine/Source/Editor/WorldPartitionEditor/Private/WorldPartition/SWorldPartitionEditorGrid2D.cpp b/Engine/Source/Editor/WorldPartitionEditor/Private/WorldPartition/SWorldPartitionEditorGrid2D.cpp
--- a/Engine/Source/Editor/WorldPartitionEditor/Private/WorldPartition/SWorldPartitionEditorGrid2D.cpp
+++ b/Engine/Source/Editor/WorldPartitionEditor/Private/WorldPartition/SWorldPartitionEditorGrid2D.cpp
@@ -605,7 +605,7 @@ void SWorldPartitionEditorGrid2D::BindCommands()
 	CommandList->MapAction(Commands.BugItGoLoadRegion, FExecuteAction::CreateLambda([this]() { GetMutableDefault<UWorldPartitionEditorPerProjectUserSettings>()->SetBugItGoLoadRegion(!GetMutableDefault<UWorldPartitionEditorPerProjectUserSettings>()->GetBugItGoLoadRegion()); }), FCanExecuteAction(), FIsActionChecked::CreateLambda([this]() { return GetMutableDefault<UWorldPartitionEditorPerProjectUserSettings>()->GetBugItGoLoadRegion(); }), FIsActionButtonVisible::CreateLambda([this]() { return GetDefault<UWorldPartitionEditorSettings>()->bEnableLoadingInEditor; }));
 
 	// Show toggles
-	CommandList->MapAction(Commands.ShowActors, FExecuteAction::CreateLambda([this]() { bShowActors = !bShowActors; }), FCanExecuteAction(), FIsActionChecked::CreateLambda([this]() { return bShowActors; }));
+	CommandList->MapAction(Commands.ShowActors, FExecuteAction::CreateLambda([this]() { bShowActors = !bShowActors; InvalidateShownActorsCache(); }), FCanExecuteAction(), FIsActionChecked::CreateLambda([this]() { return bShowActors; }));
 	CommandList->MapAction(Commands.ShowGrid, FExecuteAction::CreateLambda([this]() { bShowGrid = !bShowGrid; }), FCanExecuteAction(), FIsActionChecked::CreateLambda([this]() { return bShowGrid; }));
 	CommandList->MapAction(Commands.ShowMiniMap, FExecuteAction::CreateLambda([this]() { bShowMiniMap = !bShowMiniMap; }), FCanExecuteAction(), FIsActionChecked::CreateLambda([this]() { return bShowMiniMap; }));
 	CommandList->MapAction(Commands.ShowCoords, FExecuteAction::CreateLambda([this]() { GetMutableDefault<UWorldPartitionEditorPerProjectUserSettings>()->SetShowCellCoords(!GetMutableDefault<UWorldPartitionEditorPerProjectUserSettings>()->GetShowCellCoords()); }), FCanExecuteAction(), FIsActionChecked::CreateLambda([this]() { return GetMutableDefault<UWorldPartitionEditorPerProjectUserSettings>()->GetShowCellCoords(); }), FIsActionButtonVisible::CreateLambda([this]() { return (GetWorldPartition() && GetWorldPartition()->IsStreamingEnabled()); }));
@@ -1218,27 +1218,38 @@ void SWorldPartitionEditorGrid2D::Tick(const FGeometry& AllottedGeometry, const
 	UWorldPartitionEditorHash::FForEachIntersectingActorParams ForEachIntersectingActorParams = UWorldPartitionEditorHash::FForEachIntersectingActorParams()
 		.SetMinimumBox(FBox(FVector::ZeroVector, FVector(GetSelectionSnap())));
 
-	GetWorldPartition()->EditorHash->ForEachIntersectingActor(ViewRectWorld, [&](FWorldPartitionActorDescInstance* ActorDescInstance)
+	if (ViewRectWorld != ViewRectWorldCache)
 	{
-		if (bShowActors)
+		ShownActorGuidsCache.Reset();
+		ShownLoaderInterfacesCache.Reset();
+
+		GetWorldPartition()->EditorHash->ForEachIntersectingActor(ViewRectWorld, [&](FWorldPartitionActorDescInstance* ActorDescInstance)
 		{
-			if (ActorDescInstance->IsListedInSceneOutliner() && (!GetWorldPartition()->IsStreamingEnabled() || ActorDescInstance->GetIsSpatiallyLoaded()))
+			if (bShowActors)
 			{
-				ShownActorGuids.Add(ActorDescInstance->GetGuid());
+				if (ActorDescInstance->IsListedInSceneOutliner() && (!GetWorldPartition()->IsStreamingEnabled() || ActorDescInstance->GetIsSpatiallyLoaded()))
+				{
+					ShownActorGuidsCache.Add(ActorDescInstance->GetGuid());
+				}
 			}
-		}
 
-		if (ActorDescInstance->GetActorNativeClass()->ImplementsInterface(UWorldPartitionActorLoaderInterface::StaticClass()))
-		{
-			if (AActor* Actor = ActorDescInstance->GetActor())
+			if (ActorDescInstance->GetActorNativeClass()->ImplementsInterface(UWorldPartitionActorLoaderInterface::StaticClass()))
 			{
-				if (IWorldPartitionActorLoaderInterface::ILoaderAdapter* LoaderAdapter = Cast<IWorldPartitionActorLoaderInterface>(Actor)->GetLoaderAdapter())
+				if (AActor* Actor = ActorDescInstance->GetActor())
 				{
-					ShownLoaderInterfaces.Add(Actor);
+					if (IWorldPartitionActorLoaderInterface::ILoaderAdapter* LoaderAdapter = Cast<IWorldPartitionActorLoaderInterface>(Actor)->GetLoaderAdapter())
+					{
+						ShownLoaderInterfacesCache.Add(Actor);
+					}
 				}
 			}
-		}
-	}, ForEachIntersectingActorParams);
+		}, ForEachIntersectingActorParams);
+
+		ViewRectWorldCache = ViewRectWorld;
+	}
+
+	ShownActorGuids.Append(ShownActorGuidsCache);
+	ShownLoaderInterfaces.Append(ShownLoaderInterfacesCache);
 
 	// Also include dirty actors so we can display updated bounds.
 	if (ExternalDirtyActorsTracker.IsValid())
diff --git a/Engine/Source/Editor/WorldPartitionEditor/Public/WorldPartition/SWorldPartitionEditorGrid2D.h b/Engine/Source/Editor/WorldPartitionEditor/Public/WorldPartition/SWorldPartitionEditorGrid2D.h
--- a/Engine/Source/Editor/WorldPartitionEditor/Public/WorldPartition/SWorldPartitionEditorGrid2D.h
+++ b/Engine/Source/Editor/WorldPartitionEditor/Public/WorldPartition/SWorldPartitionEditorGrid2D.h
@@ -211,6 +211,12 @@ protected:
 	FLoaderInterfaceStack HoveredLoaderInterfacesStack;
 	FLoaderInterface HoveredLoaderInterface;
 
+	// Shown actors cache
+	FBox ViewRectWorldCache;
+	TSet<FGuid> ShownActorGuidsCache;
+	FLoaderInterfaceSet ShownLoaderInterfacesCache;
+	void InvalidateShownActorsCache() { ViewRectWorldCache.Init(); }
+
 	// Minimap
 	void UpdateWorldMiniMapDetails();
 	void SaveMiniMapUnloadedOpacityUserSetting();
