World Partition - HLOD: Added HLOD Relevant actor coloration mode

diff --git a/Engine/Source/Editor/WorldPartitionEditor/Private/WorldPartition/HLOD/HLODEditorSubsystem.cpp b/Engine/Source/Editor/WorldPartitionEditor/Private/WorldPartition/HLOD/HLODEditorSubsystem.cpp
--- a/Engine/Source/Editor/WorldPartitionEditor/Private/WorldPartition/HLOD/HLODEditorSubsystem.cpp
+++ b/Engine/Source/Editor/WorldPartitionEditor/Private/WorldPartition/HLOD/HLODEditorSubsystem.cpp
@@ -5,6 +5,7 @@
 #include "Editor.h"
 #include "Engine/Engine.h"
 #include "Engine/World.h"
+#include "GameFramework/ActorPrimitiveColorHandler.h"
 #include "GameFramework/WorldSettings.h"
 #include "Subsystems/UnrealEditorSubsystem.h"
 #include "WorldPartitionEditorModule.h"
@@ -17,13 +18,38 @@ static TAutoConsoleVariable<bool> CVarHLODInEditorEnabled(
 	true,
 	TEXT("Allow showing World Partition HLODs in the editor."));
 
+#define LOCTEXT_NAMESPACE "HLODEditorSubsystem"
+
+static FName NAME_HLODRelevantColorHandler(TEXT("HLODRelevantColorHandler"));
 
 UWorldPartitionHLODEditorSubsystem::UWorldPartitionHLODEditorSubsystem()
 {
+#if ENABLE_ACTOR_PRIMITIVE_COLOR_HANDLER
+	if (HasAnyFlags(RF_ClassDefaultObject) && ExactCast<UWorldPartitionHLODEditorSubsystem>(this))
+	{
+		FActorPrimitiveColorHandler::Get().RegisterPrimitiveColorHandler(NAME_HLODRelevantColorHandler, LOCTEXT("HLODRelevantColor", "HLOD Relevant Color"), [](const UPrimitiveComponent* InPrimitiveComponent) -> FLinearColor
+		{
+			if (AActor* Actor = InPrimitiveComponent->GetOwner())
+			{
+				if (InPrimitiveComponent->IsHLODRelevant() && Actor->IsHLODRelevant())
+				{
+					return FLinearColor::Green;
+				}
+			}
+			return FLinearColor::Red;
+		});
+	}
+#endif
 }
 
 UWorldPartitionHLODEditorSubsystem::~UWorldPartitionHLODEditorSubsystem()
 {
+#if ENABLE_ACTOR_PRIMITIVE_COLOR_HANDLER
+	if (HasAnyFlags(RF_ClassDefaultObject) && ExactCast<UWorldPartitionHLODEditorSubsystem>(this))
+	{
+		FActorPrimitiveColorHandler::Get().UnregisterPrimitiveColorHandler(NAME_HLODRelevantColorHandler);
+	}
+#endif
 }
 
 bool UWorldPartitionHLODEditorSubsystem::IsHLODInEditorEnabled()
@@ -52,11 +78,16 @@ void UWorldPartitionHLODEditorSubsystem::Initialize(FSubsystemCollectionBase& Co
 	CachedHLODMinDrawDistance = 0;
 	CachedHLODMaxDrawDistance = 0;
 	bCachedShowHLODsOverLoadedRegions = false;
+	bRefreshPrimitiveColorHandlerOnNextTick = false;
 	
 	GetWorld()->OnWorldPartitionInitialized().AddUObject(this, &UWorldPartitionHLODEditorSubsystem::OnWorldPartitionInitialized);
 	GetWorld()->OnWorldPartitionUninitialized().AddUObject(this, &UWorldPartitionHLODEditorSubsystem::OnWorldPartitionUninitialized);
 
 	GEngine->OnLevelActorListChanged().AddUObject(this, &UWorldPartitionHLODEditorSubsystem::ForceHLODStateUpdate);
+
+#if ENABLE_ACTOR_PRIMITIVE_COLOR_HANDLER
+	FCoreUObjectDelegates::OnObjectPropertyChanged.AddUObject(this, &UWorldPartitionHLODEditorSubsystem::OnObjectPostEditChange);
+#endif
 }
 
 void UWorldPartitionHLODEditorSubsystem::Deinitialize()
@@ -92,6 +123,22 @@ void UWorldPartitionHLODEditorSubsystem::OnWorldPartitionUninitialized(UWorldPar
 	}
 }
 
+void UWorldPartitionHLODEditorSubsystem::OnObjectPostEditChange(UObject* Object, FPropertyChangedEvent& PropertyChangedEvent)
+{
+#if ENABLE_ACTOR_PRIMITIVE_COLOR_HANDLER
+	if (FActorPrimitiveColorHandler::Get().GetActivePrimitiveColorHandler() == NAME_HLODRelevantColorHandler)
+	{
+		if (UPrimitiveComponent* PrimitiveComponent = Cast<UPrimitiveComponent>(Object))
+		{
+			if (PrimitiveComponent->IsRegistered())
+			{
+				RefreshPrimitivesColorsOnNextTick.Add(PrimitiveComponent);
+			}
+		}
+	}
+#endif
+}
+
 void UWorldPartitionHLODEditorSubsystem::OnLoaderAdapterStateChanged(const IWorldPartitionActorLoaderInterface::ILoaderAdapter* LoaderAdapter)
 {
 	TRACE_CPUPROFILER_EVENT_SCOPE(UWorldPartitionHLODEditorSubsystem::OnLoaderAdapterStateChanged);
@@ -111,6 +158,14 @@ void UWorldPartitionHLODEditorSubsystem::Tick(float DeltaTime)
 {
 	TRACE_CPUPROFILER_EVENT_SCOPE(UWorldPartitionHLODEditorSubsystem::Tick);
 
+#if ENABLE_ACTOR_PRIMITIVE_COLOR_HANDLER
+	if (!RefreshPrimitivesColorsOnNextTick.IsEmpty())
+	{
+		FActorPrimitiveColorHandler::Get().RefreshPrimitiveColorHandler(NAME_HLODRelevantColorHandler, RefreshPrimitivesColorsOnNextTick.Array());
+		RefreshPrimitivesColorsOnNextTick.Reset();
+	}
+#endif
+
 	if (HLODEditorData)
 	{
 		HLODEditorData->SetHLODLoadingState(IsHLODInEditorEnabled());
@@ -182,3 +237,5 @@ TStatId UWorldPartitionHLODEditorSubsystem::GetStatId() const
 {
 	RETURN_QUICK_DECLARE_CYCLE_STAT(WorldPartitionHLODEditorSubsystem, STATGROUP_Tickables);
 }
+
+#undef LOCTEXT_NAMESPACE
diff --git a/Engine/Source/Editor/WorldPartitionEditor/Public/WorldPartition/HLOD/HLODEditorSubsystem.h b/Engine/Source/Editor/WorldPartitionEditor/Public/WorldPartition/HLOD/HLODEditorSubsystem.h
--- a/Engine/Source/Editor/WorldPartitionEditor/Public/WorldPartition/HLOD/HLODEditorSubsystem.h
+++ b/Engine/Source/Editor/WorldPartitionEditor/Public/WorldPartition/HLOD/HLODEditorSubsystem.h
@@ -12,6 +12,7 @@
 
 class AActor;
 class AWorldPartitionHLOD;
+class UPrimitiveComponent;
 class UWorldPartition;
 struct FWorldPartitionHLODEditorData;
 
@@ -50,6 +51,8 @@ private:
 	void OnWorldPartitionInitialized(UWorldPartition* InWorldPartition);
 	void OnWorldPartitionUninitialized(UWorldPartition* InWorldPartition);
 
+	void OnObjectPostEditChange(UObject* Object, FPropertyChangedEvent& PropertyChangedEvent);
+
 	void OnLoaderAdapterStateChanged(const IWorldPartitionActorLoaderInterface::ILoaderAdapter* LoaderAdapter);
 
 	void ForceHLODStateUpdate();
@@ -59,8 +62,11 @@ private:
 	double CachedHLODMinDrawDistance;
 	double CachedHLODMaxDrawDistance;
 	bool bCachedShowHLODsOverLoadedRegions;
-
+	bool bRefreshPrimitiveColorHandlerOnNextTick;
 	bool bForceHLODStateUpdate;
 
+	UPROPERTY(Transient)
+	TSet<TObjectPtr<UPrimitiveComponent>> RefreshPrimitivesColorsOnNextTick;
+
 	TPimplPtr<FWorldPartitionHLODEditorData> HLODEditorData;
 };
diff --git a/Engine/Source/Runtime/Engine/Private/Actor.cpp b/Engine/Source/Runtime/Engine/Private/Actor.cpp
--- a/Engine/Source/Runtime/Engine/Private/Actor.cpp
+++ b/Engine/Source/Runtime/Engine/Private/Actor.cpp
@@ -6000,7 +6000,11 @@ bool AActor::IsHLODRelevant() const
 		return false;
 	}
 
-	if (HasAnyFlags(RF_Transient))
+	if (HasAnyFlags(RF_Transient)
+#if WITH_EDITOR 
+		&& !IsInLevelInstance()			// Treat actors in LI as HLOD relevant for the sake of visualisation modes
+#endif
+	   )
 	{
 		return false;
 	}
diff --git a/Engine/Source/Runtime/Engine/Private/Components/StaticMeshComponent.cpp b/Engine/Source/Runtime/Engine/Private/Components/StaticMeshComponent.cpp
--- a/Engine/Source/Runtime/Engine/Private/Components/StaticMeshComponent.cpp
+++ b/Engine/Source/Runtime/Engine/Private/Components/StaticMeshComponent.cpp
@@ -3030,7 +3030,11 @@ void UStaticMeshComponent::ApplyComponentInstanceData(FStaticMeshComponentInstan
 
 bool UStaticMeshComponent::IsHLODRelevant() const
 {
-	if (HasAnyFlags(RF_Transient))
+	if (HasAnyFlags(RF_Transient)
+#if WITH_EDITOR 
+		&& !(GetOwner() && GetOwner()->IsInLevelInstance())			// Treat components in LI as HLOD relevant for the sake of visualisation modes
+#endif
+		)
 	{
 		return false;
 	}
