World Partition Editor - Added option to show/hide HLOD actors bounds, which are now hidden by default

diff --git a/Engine/Source/Editor/EditorStyle/Private/StarshipStyle.cpp b/Engine/Source/Editor/EditorStyle/Private/StarshipStyle.cpp
--- a/Engine/Source/Editor/EditorStyle/Private/StarshipStyle.cpp
+++ b/Engine/Source/Editor/EditorStyle/Private/StarshipStyle.cpp
@@ -1991,6 +1991,7 @@ void FStarshipEditorStyle::FStyle::SetupWorldPartitionStyles()
 	Set("WorldPartition.BugItGoLoadRegion", new IMAGE_BRUSH_SVG("WorldPartition//Bug_20", Icon16x16)); 
 
 	Set("WorldPartition.ShowActors", new IMAGE_BRUSH_SVG("WorldPartition/Actor_20", Icon16x16));
+	Set("WorldPartition.ShowHLODActors", new IMAGE_BRUSH_SVG("Starship/Common/HierarchicalLOD", Icon16x16));
 	Set("WorldPartition.ShowGrid", new IMAGE_BRUSH_SVG("Starship/EditorViewport/grid", Icon16x16));
 	Set("WorldPartition.ShowMiniMap", new IMAGE_BRUSH_SVG("Starship/AssetIcons/Texture2D_16", Icon16x16)); 
 	Set("WorldPartition.ShowCoords", new IMAGE_BRUSH_SVG("WorldPartition/Coordinate_20", Icon16x16));
diff --git a/Engine/Source/Editor/WorldPartitionEditor/Private/WorldPartition/SWorldPartitionEditorGrid2D.cpp b/Engine/Source/Editor/WorldPartitionEditor/Private/WorldPartition/SWorldPartitionEditorGrid2D.cpp
--- a/Engine/Source/Editor/WorldPartitionEditor/Private/WorldPartition/SWorldPartitionEditorGrid2D.cpp
+++ b/Engine/Source/Editor/WorldPartitionEditor/Private/WorldPartition/SWorldPartitionEditorGrid2D.cpp
@@ -29,6 +29,7 @@
 #include "Widgets/Layout/SUniformGridPanel.h"
 #include "Framework/MultiBox/MultiBoxBuilder.h"
 #include "WorldBrowserModule.h"
+#include "WorldPartition/HLOD/HLODActor.h"
 #include "WorldPartition/LoaderAdapter/LoaderAdapterShape.h"
 #include "WorldPartition/WorldPartition.h"
 #include "WorldPartition/WorldPartitionLog.h"
@@ -239,6 +240,7 @@ void SWorldPartitionEditorGrid2D::FEditorCommands::RegisterCommands()
 	UI_COMMAND(BugItGoLoadRegion, "BugItGo Load Region", "Using BugItGo command, it will create a loading region and zoom on it.", EUserInterfaceActionType::ToggleButton, FInputChord());
 
 	UI_COMMAND(ShowActors, "Actor(s)", "Show Actor(s).", EUserInterfaceActionType::ToggleButton, FInputChord());
+	UI_COMMAND(ShowHLODActors, "HLOD Actor(s)", "Show HLOD Actor(s).", EUserInterfaceActionType::ToggleButton, FInputChord());
 	UI_COMMAND(ShowGrid, "Grid", "Show Grid.", EUserInterfaceActionType::ToggleButton, FInputChord());
 	UI_COMMAND(ShowMiniMap, "Minimap", "Show the minimap texture.", EUserInterfaceActionType::ToggleButton, FInputChord());
 	UI_COMMAND(ShowCoords, "Coordinates", "Show grid cell coordinates, you might need to zoom closer to see them.", EUserInterfaceActionType::ToggleButton, FInputChord());
@@ -397,6 +399,7 @@ TSharedRef<SWidget> SWorldPartitionEditorGrid2D::SToolBar::GenerateShowMenu() co
 
 	FToolMenuSection& Section = ShowMenu->FindOrAddSection(SectionName);
 	Section.AddMenuEntry(Commands.ShowActors);
+	Section.AddMenuEntry(Commands.ShowHLODActors);
 	Section.AddMenuEntry(Commands.ShowGrid);
 	Section.AddMenuEntry(Commands.ShowMiniMap);
 	Section.AddMenuEntry(Commands.ShowCoords);
@@ -431,6 +434,7 @@ SWorldPartitionEditorGrid2D::SWorldPartitionEditorGrid2D()
 	, bIsPanning(false)
 	, bIsMeasuring(false)
 	, bShowActors(false)
+	, bShowHLODActors(false)
 	, bShowGrid(true)
 	, bShowMiniMap(true)
 	, bFollowPlayerInPIE(false)
@@ -606,6 +610,7 @@ void SWorldPartitionEditorGrid2D::BindCommands()
 
 	// Show toggles
 	CommandList->MapAction(Commands.ShowActors, FExecuteAction::CreateLambda([this]() { bShowActors = !bShowActors; InvalidateShownActorsCache(); }), FCanExecuteAction(), FIsActionChecked::CreateLambda([this]() { return bShowActors; }));
+	CommandList->MapAction(Commands.ShowHLODActors, FExecuteAction::CreateLambda([this]() { bShowHLODActors = !bShowHLODActors; InvalidateShownActorsCache(); }), FCanExecuteAction(), FIsActionChecked::CreateLambda([this]() { return bShowHLODActors; }));
 	CommandList->MapAction(Commands.ShowGrid, FExecuteAction::CreateLambda([this]() { bShowGrid = !bShowGrid; }), FCanExecuteAction(), FIsActionChecked::CreateLambda([this]() { return bShowGrid; }));
 	CommandList->MapAction(Commands.ShowMiniMap, FExecuteAction::CreateLambda([this]() { bShowMiniMap = !bShowMiniMap; }), FCanExecuteAction(), FIsActionChecked::CreateLambda([this]() { return bShowMiniMap; }));
 	CommandList->MapAction(Commands.ShowCoords, FExecuteAction::CreateLambda([this]() { GetMutableDefault<UWorldPartitionEditorPerProjectUserSettings>()->SetShowCellCoords(!GetMutableDefault<UWorldPartitionEditorPerProjectUserSettings>()->GetShowCellCoords()); }), FCanExecuteAction(), FIsActionChecked::CreateLambda([this]() { return GetMutableDefault<UWorldPartitionEditorPerProjectUserSettings>()->GetShowCellCoords(); }), FIsActionButtonVisible::CreateLambda([this]() { return (GetWorldPartition() && GetWorldPartition()->IsStreamingEnabled()); }));
@@ -1190,6 +1195,16 @@ int32 SWorldPartitionEditorGrid2D::PaintGrid(const FGeometry& AllottedGeometry,
 	return LayerId + 1;
 }
 
+bool SWorldPartitionEditorGrid2D::ShouldShowActorBounds(AActor* InActor) const
+{
+	return bShowActors && (bShowHLODActors || !InActor->IsA<AWorldPartitionHLOD>());
+}
+
+bool SWorldPartitionEditorGrid2D::ShouldShowActorBounds(FWorldPartitionActorDescInstance* ActorDescInstance) const
+{
+	return bShowActors && (bShowHLODActors || !ActorDescInstance->GetActorNativeClass()->IsChildOf<AWorldPartitionHLOD>());
+}
+
 void SWorldPartitionEditorGrid2D::Tick(const FGeometry& AllottedGeometry, const double InCurrentTime, const float InDeltaTime)
 {
 	FWeightedMovingAverageScope ProfileMeanValue(TickTime);
@@ -1225,7 +1240,7 @@ void SWorldPartitionEditorGrid2D::Tick(const FGeometry& AllottedGeometry, const
 
 		GetWorldPartition()->EditorHash->ForEachIntersectingActor(ViewRectWorld, [&](FWorldPartitionActorDescInstance* ActorDescInstance)
 		{
-			if (bShowActors)
+			if (ShouldShowActorBounds(ActorDescInstance))
 			{
 				if (ActorDescInstance->IsListedInSceneOutliner() && (!GetWorldPartition()->IsStreamingEnabled() || ActorDescInstance->GetIsSpatiallyLoaded()))
 				{
@@ -1270,7 +1285,7 @@ void SWorldPartitionEditorGrid2D::Tick(const FGeometry& AllottedGeometry, const
 					}
 				}
 
-				if (bShowActors)
+				if (ShouldShowActorBounds(WeakActor.Get()))
 				{
 					ShownActorGuids.Add(ActorGuid);
 				}
@@ -1414,7 +1429,7 @@ uint32 SWorldPartitionEditorGrid2D::PaintActors(const FGeometry& AllottedGeometr
 	{
 		for (auto& [WeakActor, ActorGuid] : ExternalDirtyActorsTracker->GetDirtyActors())
 		{
-			if (WeakActor.IsValid() && (bShowActors || SelectedActorGuids.Contains(ActorGuid)) && !ThisWorldPartition->GetActorDescInstance(ActorGuid))
+			if (WeakActor.IsValid() && (ShouldShowActorBounds(WeakActor.Get()) || SelectedActorGuids.Contains(ActorGuid)) && !ThisWorldPartition->GetActorDescInstance(ActorGuid))
 			{
 				ActorBoundsDescs.Emplace(nullptr, WeakActor.Get());
 			}
diff --git a/Engine/Source/Editor/WorldPartitionEditor/Public/WorldPartition/SWorldPartitionEditorGrid2D.h b/Engine/Source/Editor/WorldPartitionEditor/Public/WorldPartition/SWorldPartitionEditorGrid2D.h
--- a/Engine/Source/Editor/WorldPartitionEditor/Public/WorldPartition/SWorldPartitionEditorGrid2D.h
+++ b/Engine/Source/Editor/WorldPartitionEditor/Public/WorldPartition/SWorldPartitionEditorGrid2D.h
@@ -42,6 +42,7 @@ protected:
 
 		// Show toggles
 		TSharedPtr<FUICommandInfo> ShowActors;
+		TSharedPtr<FUICommandInfo> ShowHLODActors;
 		TSharedPtr<FUICommandInfo> ShowGrid;
 		TSharedPtr<FUICommandInfo> ShowMiniMap;
 		TSharedPtr<FUICommandInfo> ShowCoords;
@@ -136,6 +137,9 @@ protected:
 	void UpdateSelectionBox(bool bSnap);
 	void ClearSelection();
 
+	bool ShouldShowActorBounds(AActor* InActor) const;
+	bool ShouldShowActorBounds(FWorldPartitionActorDescInstance* ActorDescInstance) const;
+
 	const TSharedRef<FUICommandList> CommandList;
 
 	FSingleWidgetChildrenWithBasicLayoutSlot ChildSlot;
@@ -163,6 +167,7 @@ protected:
 	bool bIsPanning;
 	bool bIsMeasuring;
 	bool bShowActors;
+	bool bShowHLODActors;
 	bool bShowGrid;
 	bool bShowMiniMap;
 	bool bFollowPlayerInPIE;
