World Partition Editor: move more frequent actions on top in  the right-click menu.

diff --git a/Engine/Source/Editor/WorldPartitionEditor/Private/WorldPartition/SWorldPartitionEditorGrid2D.cpp b/Engine/Source/Editor/WorldPartitionEditor/Private/WorldPartition/SWorldPartitionEditorGrid2D.cpp
--- a/Engine/Source/Editor/WorldPartitionEditor/Private/WorldPartition/SWorldPartitionEditorGrid2D.cpp
+++ b/Engine/Source/Editor/WorldPartitionEditor/Private/WorldPartition/SWorldPartitionEditorGrid2D.cpp
@@ -916,12 +916,13 @@ TSharedRef<SWidget> SWorldPartitionEditorGrid2D::GenerateContextualMenu() const
 		static const FName SectionSelectionName(TEXT("ContextMenu.Selection"));
 		FToolMenuSection& SectionSelection = ConxtextualMenu->AddSection(SectionSelectionName, LOCTEXT("WorldPartitionSelectionHeader", "Selection"));
 		SectionSelection.AddMenuEntry(Commands.CreateRegionFromSelection);
-		SectionSelection.AddMenuEntry(Commands.ConvertSelectedRegionsToActors);
 		
 		static const FName SectionRegionsName(TEXT("ContextMenu.Regions"));
 		FToolMenuSection& SectionRegions = ConxtextualMenu->AddSection(SectionRegionsName, LOCTEXT("WorldPartitionRegionsHeader", "Region(s)"));
 		SectionRegions.AddMenuEntry(Commands.LoadSelectedRegions);
 		SectionRegions.AddMenuEntry(Commands.UnloadSelectedRegions);
+		SectionRegions.AddSeparator(NAME_None);
+		SectionRegions.AddMenuEntry(Commands.ConvertSelectedRegionsToActors);
 	}
 
 	static const FName SectionMiscName(TEXT("ContextMenu.Misc"));
