HLOD in Editor: Properly unregister callbacks

#rnx

diff --git a/Engine/Source/Editor/WorldPartitionEditor/Private/WorldPartition/HLOD/HLODEditorData.cpp b/Engine/Source/Editor/WorldPartitionEditor/Private/WorldPartition/HLOD/HLODEditorData.cpp
--- a/Engine/Source/Editor/WorldPartitionEditor/Private/WorldPartition/HLOD/HLODEditorData.cpp
+++ b/Engine/Source/Editor/WorldPartitionEditor/Private/WorldPartition/HLOD/HLODEditorData.cpp
@@ -27,6 +27,12 @@ FWorldPartitionHLODEditorData::FWorldPartitionHLODEditorData(UWorldPartition* In
 	});	
 }
 
+FWorldPartitionHLODEditorData::~FWorldPartitionHLODEditorData()
+{
+	WorldPartition->OnActorDescContainerInstanceRegistered.RemoveAll(this);
+	WorldPartition->OnActorDescContainerInstanceUnregistered.RemoveAll(this);
+}
+
 void FWorldPartitionHLODEditorData::OnActorDescContainerInstanceRegistered(UActorDescContainerInstance* InContainerInstance)
 {
 	check(!PerContainerInstanceHLODActorDataMap.Contains(InContainerInstance));
diff --git a/Engine/Source/Editor/WorldPartitionEditor/Private/WorldPartition/HLOD/HLODEditorData.h b/Engine/Source/Editor/WorldPartitionEditor/Private/WorldPartition/HLOD/HLODEditorData.h
--- a/Engine/Source/Editor/WorldPartitionEditor/Private/WorldPartition/HLOD/HLODEditorData.h
+++ b/Engine/Source/Editor/WorldPartitionEditor/Private/WorldPartition/HLOD/HLODEditorData.h
@@ -28,6 +28,7 @@ struct FWorldPartitionHLODEditorData
 {
 public:
 	FWorldPartitionHLODEditorData(UWorldPartition* InWorldPartition);
+	~FWorldPartitionHLODEditorData();
 	
 	void ClearLoadedActorsState();
 	void UpdateLoadedActorsState();
