World Partition - HLOD: Removed unneeded manual refresh of actor coloration color

diff --git a/Engine/Source/Editor/WorldPartitionEditor/Private/WorldPartition/HLOD/HLODEditorSubsystem.cpp b/Engine/Source/Editor/WorldPartitionEditor/Private/WorldPartition/HLOD/HLODEditorSubsystem.cpp
--- a/Engine/Source/Editor/WorldPartitionEditor/Private/WorldPartition/HLOD/HLODEditorSubsystem.cpp
+++ b/Engine/Source/Editor/WorldPartitionEditor/Private/WorldPartition/HLOD/HLODEditorSubsystem.cpp
@@ -78,16 +78,11 @@ void UWorldPartitionHLODEditorSubsystem::Initialize(FSubsystemCollectionBase& Co
 	CachedHLODMinDrawDistance = 0;
 	CachedHLODMaxDrawDistance = 0;
 	bCachedShowHLODsOverLoadedRegions = false;
-	bRefreshPrimitiveColorHandlerOnNextTick = false;
 	
 	GetWorld()->OnWorldPartitionInitialized().AddUObject(this, &UWorldPartitionHLODEditorSubsystem::OnWorldPartitionInitialized);
 	GetWorld()->OnWorldPartitionUninitialized().AddUObject(this, &UWorldPartitionHLODEditorSubsystem::OnWorldPartitionUninitialized);
 
 	GEngine->OnLevelActorListChanged().AddUObject(this, &UWorldPartitionHLODEditorSubsystem::ForceHLODStateUpdate);
-
-#if ENABLE_ACTOR_PRIMITIVE_COLOR_HANDLER
-	FCoreUObjectDelegates::OnObjectPropertyChanged.AddUObject(this, &UWorldPartitionHLODEditorSubsystem::OnObjectPostEditChange);
-#endif
 }
 
 void UWorldPartitionHLODEditorSubsystem::Deinitialize()
@@ -123,22 +118,6 @@ void UWorldPartitionHLODEditorSubsystem::OnWorldPartitionUninitialized(UWorldPar
 	}
 }
 
-void UWorldPartitionHLODEditorSubsystem::OnObjectPostEditChange(UObject* Object, FPropertyChangedEvent& PropertyChangedEvent)
-{
-#if ENABLE_ACTOR_PRIMITIVE_COLOR_HANDLER
-	if (FActorPrimitiveColorHandler::Get().GetActivePrimitiveColorHandler() == NAME_HLODRelevantColorHandler)
-	{
-		if (UPrimitiveComponent* PrimitiveComponent = Cast<UPrimitiveComponent>(Object))
-		{
-			if (PrimitiveComponent->IsRegistered())
-			{
-				RefreshPrimitivesColorsOnNextTick.Add(PrimitiveComponent);
-			}
-		}
-	}
-#endif
-}
-
 void UWorldPartitionHLODEditorSubsystem::OnLoaderAdapterStateChanged(const IWorldPartitionActorLoaderInterface::ILoaderAdapter* LoaderAdapter)
 {
 	TRACE_CPUPROFILER_EVENT_SCOPE(UWorldPartitionHLODEditorSubsystem::OnLoaderAdapterStateChanged);
@@ -158,14 +137,6 @@ void UWorldPartitionHLODEditorSubsystem::Tick(float DeltaTime)
 {
 	TRACE_CPUPROFILER_EVENT_SCOPE(UWorldPartitionHLODEditorSubsystem::Tick);
 
-#if ENABLE_ACTOR_PRIMITIVE_COLOR_HANDLER
-	if (!RefreshPrimitivesColorsOnNextTick.IsEmpty())
-	{
-		FActorPrimitiveColorHandler::Get().RefreshPrimitiveColorHandler(NAME_HLODRelevantColorHandler, RefreshPrimitivesColorsOnNextTick.Array());
-		RefreshPrimitivesColorsOnNextTick.Reset();
-	}
-#endif
-
 	if (HLODEditorData)
 	{
 		HLODEditorData->SetHLODLoadingState(IsHLODInEditorEnabled());
diff --git a/Engine/Source/Editor/WorldPartitionEditor/Public/WorldPartition/HLOD/HLODEditorSubsystem.h b/Engine/Source/Editor/WorldPartitionEditor/Public/WorldPartition/HLOD/HLODEditorSubsystem.h
--- a/Engine/Source/Editor/WorldPartitionEditor/Public/WorldPartition/HLOD/HLODEditorSubsystem.h
+++ b/Engine/Source/Editor/WorldPartitionEditor/Public/WorldPartition/HLOD/HLODEditorSubsystem.h
@@ -51,8 +51,6 @@ private:
 	void OnWorldPartitionInitialized(UWorldPartition* InWorldPartition);
 	void OnWorldPartitionUninitialized(UWorldPartition* InWorldPartition);
 
-	void OnObjectPostEditChange(UObject* Object, FPropertyChangedEvent& PropertyChangedEvent);
-
 	void OnLoaderAdapterStateChanged(const IWorldPartitionActorLoaderInterface::ILoaderAdapter* LoaderAdapter);
 
 	void ForceHLODStateUpdate();
@@ -62,11 +60,7 @@ private:
 	double CachedHLODMinDrawDistance;
 	double CachedHLODMaxDrawDistance;
 	bool bCachedShowHLODsOverLoadedRegions;
-	bool bRefreshPrimitiveColorHandlerOnNextTick;
 	bool bForceHLODStateUpdate;
 
-	UPROPERTY(Transient)
-	TSet<TObjectPtr<UPrimitiveComponent>> RefreshPrimitivesColorsOnNextTick;
-
 	TPimplPtr<FWorldPartitionHLODEditorData> HLODEditorData;
 };
