World Partition Editor: allow right-click on loading regions even when fully zoomed in.

diff --git a/Engine/Source/Editor/WorldPartitionEditor/Private/WorldPartition/SWorldPartitionEditorGrid2D.cpp b/Engine/Source/Editor/WorldPartitionEditor/Private/WorldPartition/SWorldPartitionEditorGrid2D.cpp
--- a/Engine/Source/Editor/WorldPartitionEditor/Private/WorldPartition/SWorldPartitionEditorGrid2D.cpp
+++ b/Engine/Source/Editor/WorldPartitionEditor/Private/WorldPartition/SWorldPartitionEditorGrid2D.cpp
@@ -1304,7 +1304,7 @@ void SWorldPartitionEditorGrid2D::Tick(const FGeometry& AllottedGeometry, const
 
 		const IWorldPartitionActorLoaderInterface::ILoaderAdapter* LoaderAdapter = LoaderInterface->GetLoaderAdapter();
 
-		if (LoaderAdapter->GetBoundingBox().IsSet() && !LoaderAdapter->GetBoundingBox()->IsInsideXY(ViewRectWorld))
+		if (LoaderAdapter->GetBoundingBox().IsSet())
 		{
 			if (IsBoundsHovered(MouseCursorPosWorld, *LoaderAdapter->GetBoundingBox()))
 			{
