World Partition Editor
* Add ShowMouseCoordinates (World Position) in Show Menu.
* Used the default GrabHandClosed cursor when panning.

diff --git a/Engine/Source/Editor/EditorStyle/Private/StarshipStyle.cpp b/Engine/Source/Editor/EditorStyle/Private/StarshipStyle.cpp
--- a/Engine/Source/Editor/EditorStyle/Private/StarshipStyle.cpp
+++ b/Engine/Source/Editor/EditorStyle/Private/StarshipStyle.cpp
@@ -2010,7 +2010,8 @@ void FStarshipEditorStyle::FStyle::SetupWorldPartitionStyles()
 	Set("WorldPartition.ShowHLODActors", new IMAGE_BRUSH_SVG("Starship/Common/HierarchicalLOD", Icon16x16));
 	Set("WorldPartition.ShowGrid", new IMAGE_BRUSH_SVG("Starship/EditorViewport/grid", Icon16x16));
 	Set("WorldPartition.ShowMiniMap", new IMAGE_BRUSH_SVG("Starship/AssetIcons/Texture2D_16", Icon16x16)); 
-	Set("WorldPartition.ShowCoords", new IMAGE_BRUSH_SVG("WorldPartition/Coordinate_20", Icon16x16));
+	Set("WorldPartition.ShowCoords", new IMAGE_BRUSH_SVG("WorldPartition/Coordinate_16", Icon16x16));
+	Set("WorldPartition.ShowMouseCoords", new IMAGE_BRUSH_SVG("WorldPartition/MouseCoordinates_16", Icon16x16));
 
 	Set("WorldPartition.FocusSelection", new IMAGE_BRUSH_SVG("WorldPartition/FramingSelection_20", Icon16x16));
 	Set("WorldPartition.FocusLoadedRegions", new IMAGE_BRUSH_SVG("WorldPartition/FrameRegions_20", Icon16x16));
diff --git a/Engine/Source/Editor/WorldPartitionEditor/Private/WorldPartition/SWorldPartitionEditorGrid2D.cpp b/Engine/Source/Editor/WorldPartitionEditor/Private/WorldPartition/SWorldPartitionEditorGrid2D.cpp
--- a/Engine/Source/Editor/WorldPartitionEditor/Private/WorldPartition/SWorldPartitionEditorGrid2D.cpp
+++ b/Engine/Source/Editor/WorldPartitionEditor/Private/WorldPartition/SWorldPartitionEditorGrid2D.cpp
@@ -244,6 +244,7 @@ void SWorldPartitionEditorGrid2D::FEditorCommands::RegisterCommands()
 	UI_COMMAND(ShowGrid, "Grid", "Show Grid.", EUserInterfaceActionType::ToggleButton, FInputChord());
 	UI_COMMAND(ShowMiniMap, "Minimap", "Show the minimap texture.", EUserInterfaceActionType::ToggleButton, FInputChord());
 	UI_COMMAND(ShowCoords, "Coordinates", "Show grid cell coordinates, you might need to zoom closer to see them.", EUserInterfaceActionType::ToggleButton, FInputChord());
+	UI_COMMAND(ShowMouseCoords, "Mouse Coordinates", "Show mouse world coordinates.", EUserInterfaceActionType::ToggleButton, FInputChord());
 
 	UI_COMMAND(FocusSelection, "Focus Selection", "Focus Selection.", EUserInterfaceActionType::ToggleButton, FInputChord());
 	UI_COMMAND(FocusLoadedRegions, "Focus Loaded Regions", "Focus Loaded Regions.", EUserInterfaceActionType::ToggleButton, FInputChord());
@@ -403,6 +404,7 @@ TSharedRef<SWidget> SWorldPartitionEditorGrid2D::SToolBar::GenerateShowMenu() co
 	Section.AddMenuEntry(Commands.ShowGrid);
 	Section.AddMenuEntry(Commands.ShowMiniMap);
 	Section.AddMenuEntry(Commands.ShowCoords);
+	Section.AddMenuEntry(Commands.ShowMouseCoords);
 
 	return UToolMenus::Get()->GenerateWidget(MenuName, FToolMenuContext(CommandList));
 }
@@ -437,6 +439,7 @@ SWorldPartitionEditorGrid2D::SWorldPartitionEditorGrid2D()
 	, bShowHLODActors(false)
 	, bShowGrid(true)
 	, bShowMiniMap(true)
+	, bShowMouseCoords(false)
 	, bFollowPlayerInPIE(false)
 	, SelectBox(ForceInit)
 	, SelectBoxGridSnapped(ForceInit)
@@ -614,8 +617,9 @@ void SWorldPartitionEditorGrid2D::BindCommands()
 	CommandList->MapAction(Commands.ShowGrid, FExecuteAction::CreateLambda([this]() { bShowGrid = !bShowGrid; }), FCanExecuteAction(), FIsActionChecked::CreateLambda([this]() { return bShowGrid; }));
 	CommandList->MapAction(Commands.ShowMiniMap, FExecuteAction::CreateLambda([this]() { bShowMiniMap = !bShowMiniMap; }), FCanExecuteAction(), FIsActionChecked::CreateLambda([this]() { return bShowMiniMap; }));
 	CommandList->MapAction(Commands.ShowCoords, FExecuteAction::CreateLambda([this]() { GetMutableDefault<UWorldPartitionEditorPerProjectUserSettings>()->SetShowCellCoords(!GetMutableDefault<UWorldPartitionEditorPerProjectUserSettings>()->GetShowCellCoords()); }), FCanExecuteAction(), FIsActionChecked::CreateLambda([this]() { return GetMutableDefault<UWorldPartitionEditorPerProjectUserSettings>()->GetShowCellCoords(); }), FIsActionButtonVisible::CreateLambda([this]() { return (GetWorldPartition() && GetWorldPartition()->IsStreamingEnabled()); }));
+	CommandList->MapAction(Commands.ShowMouseCoords, FExecuteAction::CreateLambda([this]() { bShowMouseCoords = !bShowMouseCoords; }), FCanExecuteAction(), FIsActionChecked::CreateLambda([this]() { return bShowMouseCoords; }));
 
-	// Buttons
+	// Quick Actions
 	CommandList->MapAction(Commands.FocusSelection, FExecuteAction::CreateSP(this, &SWorldPartitionEditorGrid2D::FocusSelection), FCanExecuteAction::CreateLambda(CanFocusSelection));
 	CommandList->MapAction(Commands.FocusLoadedRegions, FExecuteAction::CreateSP(this, &SWorldPartitionEditorGrid2D::FocusLoadedRegions), FCanExecuteAction::CreateLambda([this]() { return IsInteractive() && GetWorldPartition() && GetWorldPartition()->HasLoadedUserCreatedRegions(); }), FIsActionChecked(), FIsActionButtonVisible::CreateLambda([this]() { return GetDefault<UWorldPartitionEditorSettings>()->bEnableLoadingInEditor; }));
 	CommandList->MapAction(Commands.FocusWorld, FExecuteAction::CreateSP(this, &SWorldPartitionEditorGrid2D::FocusWorld), FCanExecuteAction::CreateLambda([this]() { return IsInteractive(); }));
@@ -1140,7 +1144,7 @@ FReply SWorldPartitionEditorGrid2D::OnMouseWheel(const FGeometry& MyGeometry, co
 
 FCursorReply SWorldPartitionEditorGrid2D::OnCursorQuery(const FGeometry& MyGeometry, const FPointerEvent& CursorEvent) const
 {
-	return FCursorReply::Cursor(bIsPanning ? EMouseCursor::None : EMouseCursor::Default);
+	return FCursorReply::Cursor(bIsPanning ? EMouseCursor::GrabHandClosed : EMouseCursor::Default);
 }
 
 int32 SWorldPartitionEditorGrid2D::PaintGrid(const FGeometry& AllottedGeometry, const FSlateRect& MyCullingRect, FSlateWindowElementList& OutDrawElements, int32 LayerId) const
@@ -1775,6 +1779,23 @@ uint32 SWorldPartitionEditorGrid2D::PaintTextInfo(const FGeometry& AllottedGeome
 			ESlateDrawEffect::None,
 			USlateThemeManager::Get().GetColor(EStyleColor::AccentGray));
 	}
+
+	if (bShowMouseCoords)
+	{
+		const FString MouseCursorPosWorldText = FString::Printf(TEXT("(%.2f, %.2f)"), MouseCursorPosWorld.X, MouseCursorPosWorld.Y);
+		const FVector2D MouseCursorPosWorldTextSize = FSlateApplication::Get().GetRenderer()->GetFontMeasureService()->Measure(MouseCursorPosWorldText, FAppStyle::GetFontStyle("NormalFont"));
+		const FVector2D CurrentCursorSize = FSlateApplication::Get().GetCursorSize();
+		const float OneThirdCursorY = (CurrentCursorSize.Y * 0.5) - (MouseCursorPosWorldTextSize.Y * 0.3);
+
+		FSlateDrawElement::MakeText(
+			OutDrawElements,
+			LayerId,
+			AllottedGeometry.ToOffsetPaintGeometry(FVector2D(MouseCursorPos.X + (CurrentCursorSize.X * 0.5) + 6, MouseCursorPos.Y + OneThirdCursorY)),
+			MouseCursorPosWorldText,
+			FAppStyle::GetFontStyle("NormalFont"),
+			ESlateDrawEffect::None,
+			USlateThemeManager::Get().GetColor(EStyleColor::White));
+	}
 		
 	return LayerId + 1;
 }
@@ -1954,31 +1975,12 @@ int32 SWorldPartitionEditorGrid2D::OnPaint(const FPaintArgs& Args, const FGeomet
 		LayerId = PaintTextInfo(AllottedGeometry, MyCullingRect, OutDrawElements, ++LayerId);
 		LayerId = PaintViewer(AllottedGeometry, MyCullingRect, OutDrawElements, ++LayerId);
 		LayerId = PaintSelection(AllottedGeometry, MyCullingRect, OutDrawElements, ++LayerId);
-		LayerId = PaintSoftwareCursor(AllottedGeometry, MyCullingRect, OutDrawElements, ++LayerId);
 		LayerId = PaintMeasureTool(AllottedGeometry, MyCullingRect, OutDrawElements, ++LayerId);
 	}
 
 	return SWorldPartitionEditorGrid::OnPaint(Args, AllottedGeometry, MyCullingRect, OutDrawElements, LayerId, InWidgetStyle, bParentEnabled);
 }
 
-int32 SWorldPartitionEditorGrid2D::PaintSoftwareCursor(const FGeometry& AllottedGeometry, const FSlateRect& MyCullingRect, FSlateWindowElementList& OutDrawElements, int32 LayerId) const
-{
-	if (bIsPanning)
-	{
-		const FSlateBrush* Brush = FAppStyle::GetBrush(TEXT("SoftwareCursor_Grab"));
-		const FVector2D CursorSize = Brush->ImageSize / AllottedGeometry.Scale;
-
-		FSlateDrawElement::MakeBox(
-			OutDrawElements,
-			LayerId,
-			AllottedGeometry.ToPaintGeometry(CursorSize, FSlateLayoutTransform(MouseCursorPos - (CursorSize * 0.5f))),
-			Brush
-		);
-	}
-
-	return LayerId + 1;
-}
-
 int32 SWorldPartitionEditorGrid2D::PaintMinimap(const FGeometry& AllottedGeometry, const FSlateRect& MyCullingRect, FSlateWindowElementList& OutDrawElements, int32 LayerId) const
 {
 	// Draw MiniMap image if any
diff --git a/Engine/Source/Editor/WorldPartitionEditor/Public/WorldPartition/SWorldPartitionEditorGrid2D.h b/Engine/Source/Editor/WorldPartitionEditor/Public/WorldPartition/SWorldPartitionEditorGrid2D.h
--- a/Engine/Source/Editor/WorldPartitionEditor/Public/WorldPartition/SWorldPartitionEditorGrid2D.h
+++ b/Engine/Source/Editor/WorldPartitionEditor/Public/WorldPartition/SWorldPartitionEditorGrid2D.h
@@ -46,6 +46,7 @@ protected:
 		TSharedPtr<FUICommandInfo> ShowGrid;
 		TSharedPtr<FUICommandInfo> ShowMiniMap;
 		TSharedPtr<FUICommandInfo> ShowCoords;
+		TSharedPtr<FUICommandInfo> ShowMouseCoords;
 
 		// Quick Actions
 		TSharedPtr<FUICommandInfo> FocusSelection;
@@ -123,7 +124,6 @@ protected:
 	virtual uint32 PaintTextInfo(const FGeometry& AllottedGeometry, const FSlateRect& MyCullingRect, FSlateWindowElementList& OutDrawElements, uint32 LayerId) const;
 	virtual uint32 PaintViewer(const FGeometry& AllottedGeometry, const FSlateRect& MyCullingRect, FSlateWindowElementList& OutDrawElements, uint32 LayerId) const;
 	virtual uint32 PaintSelection(const FGeometry& AllottedGeometry, const FSlateRect& MyCullingRect, FSlateWindowElementList& OutDrawElements, uint32 LayerId) const;
-	virtual int32 PaintSoftwareCursor(const FGeometry& AllottedGeometry, const FSlateRect& MyCullingRect, FSlateWindowElementList& OutDrawElements, int32 LayerId) const;
 	virtual int32 PaintMinimap(const FGeometry& AllottedGeometry, const FSlateRect& MyCullingRect, FSlateWindowElementList& OutDrawElements, int32 LayerId) const;
 	virtual int32 PaintMeasureTool(const FGeometry& AllottedGeometry, const FSlateRect& MyCullingRect, FSlateWindowElementList& OutDrawElements, int32 LayerId) const;
 
@@ -170,6 +170,7 @@ protected:
 	bool bShowHLODActors;
 	bool bShowGrid;
 	bool bShowMiniMap;
+	bool bShowMouseCoords;
 	bool bFollowPlayerInPIE;
 	FVector2D MouseCursorPos;
 	FVector2D MouseCursorPosWorld;
