World Partition - HLOD: When HLOD in editor is enabled, avoid stack overflow when pinning actors with circular references

diff --git a/Engine/Source/Editor/WorldPartitionEditor/Private/WorldPartition/HLOD/HLODEditorData.cpp b/Engine/Source/Editor/WorldPartitionEditor/Private/WorldPartition/HLOD/HLODEditorData.cpp
--- a/Engine/Source/Editor/WorldPartitionEditor/Private/WorldPartition/HLOD/HLODEditorData.cpp
+++ b/Engine/Source/Editor/WorldPartitionEditor/Private/WorldPartition/HLOD/HLODEditorData.cpp
@@ -109,10 +109,18 @@ struct FBoundsWithVolume
 };
 
 // Gather Pinned Actors bounds - also include references & contained actors if the pinned actor is a container.
-static void GatherLoadedActorsBounds(TArray<FBoundsWithVolume>& OutLoadedBounds, const FWorldPartitionActorDescInstance* InActorDescInstance, const UActorDescContainerInstance* InContainerInstance, const TOptional<FTransform>& InContainerTransform = TOptional<FTransform>())
+static void GatherLoadedActorsBounds(TArray<FBoundsWithVolume>& OutLoadedBounds, TSet<const FWorldPartitionActorDescInstance*>& InProcessedActors, const FWorldPartitionActorDescInstance* InActorDescInstance, const UActorDescContainerInstance* InContainerInstance, const TOptional<FTransform>& InContainerTransform = TOptional<FTransform>())
 {
 	if (InActorDescInstance)
 	{
+		bool bIsAlreadyInSet;
+		InProcessedActors.FindOrAdd(InActorDescInstance, &bIsAlreadyInSet);
+
+		if (bIsAlreadyInSet)
+		{
+			return;
+		}
+
 		// The actor itself - Include the actor bounds only if HLOD relevant
 		if (InActorDescInstance->GetIsSpatiallyLoaded() && InActorDescInstance->GetActorIsHLODRelevant() && InActorDescInstance->IsEditorRelevant())
 		{
@@ -125,7 +133,7 @@ static void GatherLoadedActorsBounds(TArray<FBoundsWithVolume>& OutLoadedBounds,
 		{
 			if (const FWorldPartitionActorDescInstance* ReferenceActorDescInstance = InContainerInstance->GetActorDescInstance(ReferenceGuid))
 			{
-				GatherLoadedActorsBounds(OutLoadedBounds, ReferenceActorDescInstance, InContainerInstance, InContainerTransform);
+				GatherLoadedActorsBounds(OutLoadedBounds, InProcessedActors, ReferenceActorDescInstance, InContainerInstance, InContainerTransform);
 			}
 		}
 
@@ -138,7 +146,7 @@ static void GatherLoadedActorsBounds(TArray<FBoundsWithVolume>& OutLoadedBounds,
 				FTransform ContainerWorldSpaceTransform = InContainerTransform.IsSet() ? ContainerInstance.Transform * InContainerTransform.GetValue() : ContainerInstance.Transform;
 				for (UActorDescContainerInstance::TConstIterator<> Iterator(ContainerInstance.ContainerInstance); Iterator; ++Iterator)
 				{
-					GatherLoadedActorsBounds(OutLoadedBounds, *Iterator, ContainerInstance.ContainerInstance, ContainerWorldSpaceTransform);
+					GatherLoadedActorsBounds(OutLoadedBounds, InProcessedActors, *Iterator, ContainerInstance.ContainerInstance, ContainerWorldSpaceTransform);
 				}
 			}
 		}
@@ -195,11 +203,13 @@ void FWorldPartitionHLODEditorData::UpdateLoadedActorsState()
 	// Gather Pinned Actors bounds
 	if (WorldPartition->PinnedActors)
 	{
+		TSet<const FWorldPartitionActorDescInstance*> InProcessedActors;
+
 		for (const FWorldPartitionHandle& PinnedActor : WorldPartition->PinnedActors->GetActors())
 		{
 			if (PinnedActor.IsValid())
 			{
-				GatherLoadedActorsBounds(LoadedBounds, *PinnedActor, WorldPartition->GetActorDescContainerInstance());
+				GatherLoadedActorsBounds(LoadedBounds, InProcessedActors, *PinnedActor, WorldPartition->GetActorDescContainerInstance());
 			}
 		}
 	}
