World Partition
* Fix crash when accessing context menu on Location Volume or ProceduleFoliageVolume within a non World Partition world.
#tests With and without WP, opening context menu on LocationVolume & ProcFoliageVolume

diff --git a/Engine/Source/Editor/WorldPartitionEditor/Private/WorldPartitionEditorModule.cpp b/Engine/Source/Editor/WorldPartitionEditor/Private/WorldPartitionEditorModule.cpp
--- a/Engine/Source/Editor/WorldPartitionEditor/Private/WorldPartitionEditorModule.cpp
+++ b/Engine/Source/Editor/WorldPartitionEditor/Private/WorldPartitionEditorModule.cpp
@@ -75,15 +75,16 @@ static void OnSelectedWorldPartitionVolumesToggleLoading(TArray<TWeakObjectPtr<A
 	{
 		if (Actor->Implements<UWorldPartitionActorLoaderInterface>())
 		{
-			IWorldPartitionActorLoaderInterface::ILoaderAdapter* LoaderAdapter = Cast<IWorldPartitionActorLoaderInterface>(Actor)->GetLoaderAdapter();
-
-			if (bLoad)
-			{
-				LoaderAdapter->Load();
-			}
-			else
+			if (IWorldPartitionActorLoaderInterface::ILoaderAdapter* LoaderAdapter = Cast<IWorldPartitionActorLoaderInterface>(Actor)->GetLoaderAdapter())
 			{
-				LoaderAdapter->Unload();
+				if (bLoad)
+				{
+					LoaderAdapter->Load();
+				}
+				else
+				{
+					LoaderAdapter->Unload();
+				}
 			}
 		}
 	}
@@ -95,10 +96,12 @@ static bool CanLoadUnloadSelectedVolumes(TArray<TWeakObjectPtr<AActor>> Volumes,
 	{
 		if (Actor->Implements<UWorldPartitionActorLoaderInterface>())
 		{
-			IWorldPartitionActorLoaderInterface::ILoaderAdapter* LoaderAdapter = Cast<IWorldPartitionActorLoaderInterface>(Actor)->GetLoaderAdapter();
-			if (bLoad != LoaderAdapter->IsLoaded())
+			if (IWorldPartitionActorLoaderInterface::ILoaderAdapter* LoaderAdapter = Cast<IWorldPartitionActorLoaderInterface>(Actor)->GetLoaderAdapter())
 			{
-				return true;
+				if (bLoad != LoaderAdapter->IsLoaded())
+				{
+					return true;
+				}
 			}
 		}
 	}
@@ -156,6 +159,10 @@ static void CreateLevelViewportContextMenuEntries(FMenuBuilder& MenuBuilder, TAr
 					EditorLoaderAdapter->GetLoaderAdapter()->SetUserCreated(true);
 					EditorLoaderAdapter->GetLoaderAdapter()->Load();
 				}
+			}),
+			FCanExecuteAction::CreateLambda([World, WorldPartition]()
+			{
+				return (World.IsValid() && WorldPartition.IsValid());
 			})
 		);
 
