World Partition - HLOD: In editor, prevent HLODs from showing up on top of modified actors

diff --git a/Engine/Source/Editor/WorldPartitionEditor/Private/WorldPartition/HLOD/HLODEditorData.cpp b/Engine/Source/Editor/WorldPartitionEditor/Private/WorldPartition/HLOD/HLODEditorData.cpp
--- a/Engine/Source/Editor/WorldPartitionEditor/Private/WorldPartition/HLOD/HLODEditorData.cpp
+++ b/Engine/Source/Editor/WorldPartitionEditor/Private/WorldPartition/HLOD/HLODEditorData.cpp
@@ -25,6 +25,8 @@ FWorldPartitionHLODEditorData::FWorldPartitionHLODEditorData(UWorldPartition* In
 	{
 		OnActorDescContainerInstanceRegistered(InContainerInstance);
 	});	
+
+	ExternalDirtyActorsTracker = MakeUnique<FExternalDirtyActorsTracker>(WorldPartition->GetWorld()->PersistentLevel, this);
 }
 
 FWorldPartitionHLODEditorData::~FWorldPartitionHLODEditorData()
@@ -199,20 +201,31 @@ void FWorldPartitionHLODEditorData::UpdateLoadedActorsState()
 
 		return true;
 	});
-	
+
+	TSet<const FWorldPartitionActorDescInstance*> ProcessedActors;
+
 	// Gather Pinned Actors bounds
 	if (WorldPartition->PinnedActors)
 	{
-		TSet<const FWorldPartitionActorDescInstance*> InProcessedActors;
-
 		for (const FWorldPartitionHandle& PinnedActor : WorldPartition->PinnedActors->GetActors())
 		{
 			if (PinnedActor.IsValid())
 			{
-				GatherLoadedActorsBounds(LoadedBounds, InProcessedActors, *PinnedActor, WorldPartition->GetActorDescContainerInstance());
+				GatherLoadedActorsBounds(LoadedBounds, ProcessedActors, *PinnedActor, WorldPartition->GetActorDescContainerInstance());
 			}
 		}
 	}
+	
+	// Gather Dirty Actors bounds
+	for (FExternalDirtyActorsTracker::MapType::TConstIterator MapIt(ExternalDirtyActorsTracker->GetDirtyActors()); MapIt; ++MapIt)
+	{
+		const TWeakObjectPtr<AActor>& Actor = MapIt.Key();
+		if (Actor.IsValid())
+		{
+			FWorldPartitionActorDescInstance* ActorDescInstance = WorldPartition->GetActorDescInstance(Actor->GetActorGuid());
+			GatherLoadedActorsBounds(LoadedBounds, ProcessedActors, ActorDescInstance, WorldPartition->GetActorDescContainerInstance());
+		}
+	}
 
 	// Sort Bounds by volume
 	Algo::Sort(LoadedBounds, [](const FBoundsWithVolume& BoxA, const FBoundsWithVolume& BoxB) { return BoxA.Volume > BoxB.Volume; });
@@ -223,7 +236,20 @@ void FWorldPartitionHLODEditorData::UpdateLoadedActorsState()
 		{
 			FBox HLODSceneNodeBox = HLODSceneNode->Bounds.GetBox();
 
-			const bool bHasAnyIntersectingLoadedRegion = Algo::AnyOf(LoadedBounds, [&HLODSceneNodeBox](const FBoundsWithVolume& Bounds) { return Bounds.Box.Intersect(HLODSceneNodeBox); });
+			// Perform an overlap test that exclude shared edges. For example, this prevents neighbor landscape HLODs actors that are
+			// perfectly matching landscape proxy bounds from being hidden when a landscape proxy is pinned.
+			const bool bHasAnyIntersectingLoadedRegion = Algo::AnyOf(LoadedBounds, [&HLODSceneNodeBox](const FBoundsWithVolume& Bounds) 
+			{
+				const FBox BoxOverlap = Bounds.Box.Overlap(HLODSceneNodeBox);
+				if (!BoxOverlap.IsValid)
+				{
+					return false;
+				}
+
+				const FVector OverlapSize = BoxOverlap.GetSize();
+				return OverlapSize.X > UE_THRESH_POINTS_ARE_NEAR && OverlapSize.Y > UE_THRESH_POINTS_ARE_NEAR;
+			});
+
 			if (bHasAnyIntersectingLoadedRegion)
 			{
 				HLODSceneNode->HasIntersectingLoadedRegion = StateUpdate;
diff --git a/Engine/Source/Editor/WorldPartitionEditor/Private/WorldPartition/HLOD/HLODEditorData.h b/Engine/Source/Editor/WorldPartitionEditor/Private/WorldPartition/HLOD/HLODEditorData.h
--- a/Engine/Source/Editor/WorldPartitionEditor/Private/WorldPartition/HLOD/HLODEditorData.h
+++ b/Engine/Source/Editor/WorldPartitionEditor/Private/WorldPartition/HLOD/HLODEditorData.h
@@ -3,6 +3,7 @@
 #pragma once
 
 #include "WorldPartition/HLOD/HLODLoaderAdapter.h"
+#include "ExternalDirtyActorsTracker.h"
 
 class UWorldPartition;
 class UActorDescContainerInstance;
@@ -51,4 +52,14 @@ private:
 	TMap<UActorDescContainerInstance*, FContainerInstanceHLODActorData> PerContainerInstanceHLODActorDataMap;
 	TUniquePtr<FLoaderAdapterHLOD> HLODActorsLoader;
 	int32 LastStateUpdate;
+
+	struct FExternalDirtyActorTrackerGuid
+	{
+		using Type = FGuid;
+		using OwnerType = FWorldPartitionHLODEditorData;
+		static FGuid Store(FWorldPartitionHLODEditorData* InOwner, AActor* InActor) { return InActor->GetActorGuid(); }
+	};
+
+	using FExternalDirtyActorsTracker = TExternalDirtyActorsTracker<FExternalDirtyActorTrackerGuid>;
+	TUniquePtr<FExternalDirtyActorsTracker> ExternalDirtyActorsTracker;
 };
