World Partition - HLOD: Input Stats
* During HLOD generation, gather references to assets (for now, only of static meshes) and number of occurences, and store that information in HLOD actors.
* Added "wp.Editor.HLOD.DumpInputStats" console command to write a report to disk
* General cleanup - Moved other HLOD stat code out of the HLODEngineSubsystem and into the HLODEditorSubsystem

diff --git a/Engine/Plugins/Editor/WorldPartitionHLODUtilities/Source/Private/WorldPartition/HLOD/Utilities/WorldPartitionHLODUtilities.cpp b/Engine/Plugins/Editor/WorldPartitionHLODUtilities/Source/Private/WorldPartition/HLOD/Utilities/WorldPartitionHLODUtilities.cpp
--- a/Engine/Plugins/Editor/WorldPartitionHLODUtilities/Source/Private/WorldPartition/HLOD/Utilities/WorldPartitionHLODUtilities.cpp
+++ b/Engine/Plugins/Editor/WorldPartitionHLODUtilities/Source/Private/WorldPartition/HLOD/Utilities/WorldPartitionHLODUtilities.cpp
@@ -46,7 +46,7 @@ static uint32 ComputeHLODHash(AWorldPartitionHLOD* InHLODActor, const TArray<UAc
 	FArchiveCrc32 Ar;
 
 	// Base key, changing this will force a rebuild of all HLODs
-	FString HLODBaseKey = "95D53D3FC3154822BC00C466E8416BDB";
+	FString HLODBaseKey = "0D33837AB1A04CC2AEEB04C5C217DE96";
 	Ar << HLODBaseKey;
 
 	// HLOD Source Actors
@@ -776,25 +776,25 @@ uint32 FWorldPartitionHLODUtilities::BuildHLOD(AWorldPartitionHLOD* InHLODActor)
 			}
 
 			// Build
-			TArray<UActorComponent*> HLODComponents;
+			FHLODBuildResult BuildResult;
 			{
 				FAutoScopedDurationTimer BuildTimeScope;
-				HLODComponents = HLODBuilder->Build(HLODBuildContext);
+				BuildResult = HLODBuilder->Build(HLODBuildContext);
 				BuildTimeMS = FMath::RoundToInt(BuildTimeScope.GetTime() * 1000);
 			}
 
 			if (HLODModifier)
 			{
-				HLODModifier->EndHLODBuild(HLODComponents);
+				HLODModifier->EndHLODBuild(BuildResult.HLODComponents);
 			}
 
-			if (HLODComponents.IsEmpty())
+			if (BuildResult.HLODComponents.IsEmpty())
 			{
 				UE_LOG(LogHLODBuilder, Warning, TEXT("HLOD generation created no component for %s"), *InHLODActor->GetActorLabel());
 			}
 
 			// Ideally, this should be performed elsewhere, to allow more flexibility in the HLOD generation
-			for (UActorComponent* HLODComponent : HLODComponents)
+			for (UActorComponent* HLODComponent : BuildResult.HLODComponents)
 			{
 				HLODComponent->SetCanEverAffectNavigation(false);
 
@@ -859,7 +859,8 @@ uint32 FWorldPartitionHLODUtilities::BuildHLOD(AWorldPartitionHLOD* InHLODActor)
 				}
 			}
 
-			InHLODActor->SetHLODComponents(HLODComponents);
+			InHLODActor->SetInputStats(BuildResult.InputStats);
+			InHLODActor->SetHLODComponents(BuildResult.HLODComponents);
 		}
 	}
 
diff --git a/Engine/Source/Editor/UnrealEd/Private/WorldPartition/WorldPartitionHLODsBuilder.cpp b/Engine/Source/Editor/UnrealEd/Private/WorldPartition/WorldPartitionHLODsBuilder.cpp
--- a/Engine/Source/Editor/UnrealEd/Private/WorldPartition/WorldPartitionHLODsBuilder.cpp
+++ b/Engine/Source/Editor/UnrealEd/Private/WorldPartition/WorldPartitionHLODsBuilder.cpp
@@ -31,7 +31,7 @@
 #include "WorldPartition/HLOD/HLODActorDesc.h"
 #include "WorldPartition/HLOD/HLODLayer.h"
 #include "WorldPartition/HLOD/HLODProviderInterface.h"
-#include "WorldPartition/HLOD/HLODRuntimeSubsystem.h"
+#include "WorldPartition/IWorldPartitionEditorModule.h"
 
 DEFINE_LOG_CATEGORY_STATIC(LogWorldPartitionHLODsBuilder, All, All);
 
@@ -566,7 +566,12 @@ bool UWorldPartitionHLODsBuilder::SubmitHLODActors()
 bool UWorldPartitionHLODsBuilder::DumpStats()
 {
 	const FString HLODStatsOutputFilename = FPaths::ProjectSavedDir() / TEXT("WorldPartition") / FString::Printf(TEXT("HLODStats-%08x.csv"), FPlatformProcess::GetCurrentProcessId());
-	return UWorldPartitionHLODRuntimeSubsystem::WriteHLODStatsCSV(World, HLODStatsOutputFilename);
+
+	IWorldPartitionEditorModule::FWriteHLODStatsParams StatsParams;
+	StatsParams.Filename = HLODStatsOutputFilename;
+	StatsParams.World = World;
+	StatsParams.StatsType = IWorldPartitionEditorModule::FWriteHLODStatsParams::EStatsType::Default;
+	return IWorldPartitionEditorModule::Get().WriteHLODStats(StatsParams);
 }
 
 bool UWorldPartitionHLODsBuilder::GetHLODActorsToBuild(TArray<FGuid>& HLODActorsToBuild) const
diff --git a/Engine/Source/Editor/WorldPartitionEditor/Private/WorldPartition/HLOD/HLODEditorSubsystem.cpp b/Engine/Source/Editor/WorldPartitionEditor/Private/WorldPartition/HLOD/HLODEditorSubsystem.cpp
--- a/Engine/Source/Editor/WorldPartitionEditor/Private/WorldPartition/HLOD/HLODEditorSubsystem.cpp
+++ b/Engine/Source/Editor/WorldPartitionEditor/Private/WorldPartition/HLOD/HLODEditorSubsystem.cpp
@@ -7,13 +7,22 @@
 #include "Engine/World.h"
 #include "GameFramework/ActorPrimitiveColorHandler.h"
 #include "GameFramework/WorldSettings.h"
+#include "Misc/FileHelper.h"
+#include "StaticMeshResources.h"
 #include "Subsystems/UnrealEditorSubsystem.h"
+#include "UObject/UObjectIterator.h"
 #include "WorldPartitionEditorModule.h"
+#include "WorldPartition/HLOD/HLODActor.h"
+#include "WorldPartition/HLOD/HLODActorDesc.h"
 #include "WorldPartition/HLOD/HLODEditorData.h"
 #include "WorldPartition/HLOD/HLODLayer.h"
+#include "WorldPartition/HLOD/HLODStats.h"
 #include "WorldPartition/WorldPartition.h"
+#include "WorldPartition/WorldPartitionActorDescInstance.h"
 #include "WorldPartition/WorldPartitionEditorSettings.h"
 #include "WorldPartition/WorldPartitionSubsystem.h"
+#include "WorldPartition/DataLayer/WorldDataLayers.h"
+#include "WorldPartition/DataLayer/DataLayerManager.h"
 
 #include "PropertyPermissionList.h"
 
@@ -24,6 +33,8 @@ static TAutoConsoleVariable<bool> CVarHLODInEditorEnabled(
 
 #define LOCTEXT_NAMESPACE "HLODEditorSubsystem"
 
+DEFINE_LOG_CATEGORY_STATIC(LogHLODEditorSubsystem, Log, All);
+
 static FName NAME_HLODRelevantColorHandler(TEXT("HLODRelevantColorHandler"));
 TMap<EHLODSettingsVisibility, UWorldPartitionHLODEditorSubsystem::FStructsPropertiesMap> UWorldPartitionHLODEditorSubsystem::StructsPropertiesVisibility;
 
@@ -261,4 +272,247 @@ void UWorldPartitionHLODEditorSubsystem::AddHLODSettingsFilter(EHLODSettingsVisi
 	StructsPropertiesVisibility.FindOrAdd(InSettingsVisibility).FindOrAdd(InStruct).Add(InPropertyName);
 }
 
+bool UWorldPartitionHLODEditorSubsystem::WriteHLODStats(const IWorldPartitionEditorModule::FWriteHLODStatsParams& Params) const
+{
+	check(Params.World == GetWorld());
+
+	bool bResult = false;
+
+	switch(Params.StatsType)
+	{
+	case IWorldPartitionEditorModule::FWriteHLODStatsParams::EStatsType::Default:
+		bResult = WriteHLODStats(Params.Filename);
+		break;
+
+	case IWorldPartitionEditorModule::FWriteHLODStatsParams::EStatsType::InputDetails:
+		bResult = WriteHLODInputStats(Params.Filename);
+		break;
+	}
+
+	if (bResult)
+	{
+		UE_LOG(LogHLODEditorSubsystem, Display, TEXT("Wrote HLOD stats to %s"), *Params.Filename);
+	}
+	else
+	{
+		UE_LOG(LogHLODEditorSubsystem, Error, TEXT("Failed to write HLOD stats to %s"), *Params.Filename);
+	}
+
+	return bResult;
+}
+
+bool UWorldPartitionHLODEditorSubsystem::WriteHLODStats(const FString& InFilename) const
+{
+	UWorld* World = GetWorld();
+	UWorldPartition* WorldPartition = World->GetWorldPartition();
+	if (!WorldPartition)
+	{
+		return false;
+	}
+
+	typedef TFunction<FString(FWorldPartitionActorDescInstance* InActorDescInstance, const FHLODActorDesc&)> FGetStatFunc;
+
+	auto GetHLODStat = [](FName InStatName)
+	{
+		return TPair<FName, FGetStatFunc>(InStatName, [InStatName](FWorldPartitionActorDescInstance* InActorDescInstance, const FHLODActorDesc& InActorDesc)
+		{
+			return FString::Printf(TEXT("%lld"), InActorDesc.GetStat(InStatName));
+		});
+	};
+
+	const UDataLayerManager* DataLayerManager = WorldPartition->GetDataLayerManager();
+
+	auto GetDataLayerShortName = [DataLayerManager](FName DataLayerInstanceName)
+	{
+		const UDataLayerInstance* DataLayerInstance = DataLayerManager ? DataLayerManager->GetDataLayerInstance(DataLayerInstanceName) : nullptr;
+		return DataLayerInstance ? DataLayerInstance->GetDataLayerShortName() : DataLayerInstanceName.ToString();
+	};
+
+	TArray<TPair<FName, FGetStatFunc>> StatsToWrite =
+	{
+		{ "WorldPackage",		[World](FWorldPartitionActorDescInstance* InActorDescInstance, const FHLODActorDesc& InActorDesc) { return World->GetPackage()->GetName(); } },
+		{ "Name",				[](FWorldPartitionActorDescInstance* InActorDescInstance, const FHLODActorDesc& InActorDesc) { return InActorDescInstance->GetActorLabelString(); } },
+		{ "HLODLayer",			[](FWorldPartitionActorDescInstance* InActorDescInstance, const FHLODActorDesc& InActorDesc) { return InActorDesc.GetSourceHLODLayer().GetAssetName().ToString(); }},
+		{ "SpatiallyLoaded",	[](FWorldPartitionActorDescInstance* InActorDescInstance, const FHLODActorDesc& InActorDesc) { return InActorDescInstance->GetIsSpatiallyLoaded() ? TEXT("true") : TEXT("false"); } },
+		{ "DataLayers",			[&GetDataLayerShortName](FWorldPartitionActorDescInstance* InActorDescInstance, const FHLODActorDesc& InActorDesc) { return FString::JoinBy(InActorDescInstance->GetDataLayerInstanceNames().ToArray(), TEXT(" | "), GetDataLayerShortName); }},
+
+		GetHLODStat(FWorldPartitionHLODStats::InputActorCount),
+		GetHLODStat(FWorldPartitionHLODStats::InputTriangleCount),
+		GetHLODStat(FWorldPartitionHLODStats::InputVertexCount),
+
+		GetHLODStat(FWorldPartitionHLODStats::MeshInstanceCount),
+		GetHLODStat(FWorldPartitionHLODStats::MeshNaniteTriangleCount),
+		GetHLODStat(FWorldPartitionHLODStats::MeshNaniteVertexCount),
+		GetHLODStat(FWorldPartitionHLODStats::MeshTriangleCount),
+		GetHLODStat(FWorldPartitionHLODStats::MeshVertexCount),
+		GetHLODStat(FWorldPartitionHLODStats::MeshUVChannelCount),
+
+		GetHLODStat(FWorldPartitionHLODStats::MaterialBaseColorTextureSize),
+		GetHLODStat(FWorldPartitionHLODStats::MaterialNormalTextureSize),
+		GetHLODStat(FWorldPartitionHLODStats::MaterialEmissiveTextureSize),
+		GetHLODStat(FWorldPartitionHLODStats::MaterialMetallicTextureSize),
+		GetHLODStat(FWorldPartitionHLODStats::MaterialRoughnessTextureSize),
+		GetHLODStat(FWorldPartitionHLODStats::MaterialSpecularTextureSize),
+
+		GetHLODStat(FWorldPartitionHLODStats::MemoryMeshResourceSizeBytes),
+		GetHLODStat(FWorldPartitionHLODStats::MemoryTexturesResourceSizeBytes),
+		GetHLODStat(FWorldPartitionHLODStats::MemoryDiskSizeBytes),
+
+		GetHLODStat(FWorldPartitionHLODStats::BuildTimeLoadMilliseconds),
+		GetHLODStat(FWorldPartitionHLODStats::BuildTimeBuildMilliseconds),
+		GetHLODStat(FWorldPartitionHLODStats::BuildTimeTotalMilliseconds)
+	};
+
+	FStringOutputDevice Output;
+
+	// Write header if file doesn't exist
+	if (!IFileManager::Get().FileExists(*InFilename))
+	{
+		const FString StatHeader = FString::JoinBy(StatsToWrite, TEXT(","), [](const TPair<FName, FGetStatFunc>& Pair) { return Pair.Key.ToString(); });
+		Output.Logf(TEXT("%s" LINE_TERMINATOR_ANSI), *StatHeader);
+	}
+
+	// Write one line per HLOD actor desc
+	for (FActorDescContainerInstanceCollection::TIterator<AWorldPartitionHLOD> HLODIterator(WorldPartition); HLODIterator; ++HLODIterator)
+	{
+		const FString StatLine = FString::JoinBy(StatsToWrite, TEXT(","), [&HLODIterator](const TPair<FName, FGetStatFunc>& Pair)
+			{
+				const FHLODActorDesc& HLODActorDesc = *(FHLODActorDesc*)HLODIterator->GetActorDesc();
+				return Pair.Value(*HLODIterator, HLODActorDesc);
+			});
+		Output.Logf(TEXT("%s" LINE_TERMINATOR_ANSI), *StatLine);
+	}
+
+	// Write to file
+	return FFileHelper::SaveStringToFile(Output, *InFilename, FFileHelper::EEncodingOptions::AutoDetect, &IFileManager::Get(), EFileWrite::FILEWRITE_Append);
+}
+
+bool UWorldPartitionHLODEditorSubsystem::WriteHLODInputStats(const FString& InFilename) const
+{
+	UWorld* World = GetWorld();
+	UWorldPartition* WorldPartition = World->GetWorldPartition();
+	if (!WorldPartition)
+	{
+		return false;
+	}
+
+	FWorldPartitionHelpers::FForEachActorWithLoadingParams ForEachActorWithLoadingParams;
+	ForEachActorWithLoadingParams.ActorClasses = { AWorldPartitionHLOD::StaticClass() };
+
+	TMap<TPair<int32, FName>, FHLODBuildInputReferencedAssets> BuildersReferencedAssets;
+
+	// Aggregate referenced assets from all HLOD actors
+	FWorldPartitionHelpers::ForEachActorWithLoading(WorldPartition, [this, World, WorldPartition, &BuildersReferencedAssets](const FWorldPartitionActorDescInstance* ActorDescInstance)
+	{
+		AWorldPartitionHLOD* HLODActor = Cast<AWorldPartitionHLOD>(ActorDescInstance->GetActor());
+		if (!HLODActor)
+		{
+			UE_LOG(LogHLODEditorSubsystem, Error, TEXT("HLOD actor failed to load: %s (%s)"), *ActorDescInstance->GetActorNameString(), *ActorDescInstance->GetActorPackage().ToString());
+			return false;
+		}
+
+		const FHLODBuildInputStats& InputStats = HLODActor->GetInputStats();
+
+		for (const TPair<FName, FHLODBuildInputReferencedAssets>& Entry : InputStats.BuildersReferencedAssets)
+		{
+			FHLODBuildInputReferencedAssets& BuilderReferencedAssets = BuildersReferencedAssets.FindOrAdd(TPair<int32, FName>(HLODActor->GetLODLevel(), Entry.Key));
+			for (const TPair<FTopLevelAssetPath, uint32>& ReferencedMesh : Entry.Value.StaticMeshes)
+			{
+				BuilderReferencedAssets.StaticMeshes.FindOrAdd(ReferencedMesh.Key) += ReferencedMesh.Value;
+			}
+		}
+
+		return true;
+	}, ForEachActorWithLoadingParams);
+
+	FStringOutputDevice Output;
+
+	Output.Logf(TEXT("HLODLevel,BuilderName,AssetName,RefCount,LastLODTriCount,LastLODVtxCount" LINE_TERMINATOR_ANSI));
+
+	BuildersReferencedAssets.KeySort([](const TPair<int32, FName>& PairA, const TPair<int32, FName> PairB)
+	{
+		if (PairA.Key != PairB.Key)
+		{
+			return PairA.Key < PairB.Key;
+		}
+
+		return PairA.Value.LexicalLess(PairB.Value);
+	});
+
+	for (TPair<TPair<int32,FName>, FHLODBuildInputReferencedAssets>& Entry : BuildersReferencedAssets)
+	{
+		Entry.Value.StaticMeshes.KeySort(FTopLevelAssetPathFastLess());
+
+		for (const TPair<FTopLevelAssetPath, uint32>& ReferencedMesh : Entry.Value.StaticMeshes)
+		{
+			const FTopLevelAssetPath& StaticMeshAssetPath = ReferencedMesh.Key;
+			
+			UObject* LoadedObject = StaticLoadAsset(UObject::StaticClass(), StaticMeshAssetPath, LOAD_NoWarn);
+			if (UStaticMesh* StaticMesh = Cast<UStaticMesh>(LoadedObject))
+			{
+				const FStaticMeshRenderData* RenderData = StaticMesh->GetRenderData();
+				const bool bHasRenderData = RenderData && !RenderData->LODResources.IsEmpty();
+				if (bHasRenderData)
+				{
+					const int32 LODIndex = StaticMesh->GetNumLODs() - 1;
+					const int64 LastLODTriCount = RenderData->LODResources[LODIndex].GetNumTriangles();
+					const int64 LastLODVtxCount = RenderData->LODResources[LODIndex].GetNumVertices();
+					Output.Logf(TEXT("HLOD%d,%s,%s,%d,%d,%d" LINE_TERMINATOR_ANSI), Entry.Key.Key, *Entry.Key.Value.ToString(), *StaticMeshAssetPath.GetPackageName().ToString(), ReferencedMesh.Value, LastLODTriCount, LastLODVtxCount);
+				}				
+			}
+		}
+	}
+
+	// Write to file
+	return FFileHelper::SaveStringToFile(Output, *InFilename);
+}
+
+FAutoConsoleCommand HLODDumpStats(
+	TEXT("wp.Editor.HLOD.DumpStats"),
+	TEXT("Export various HLOD stats to a CSV formatted file."),
+	FConsoleCommandWithArgsDelegate::CreateLambda([](const TArray<FString>& Args)
+	{
+		const FString HLODStatsOutputFilename = FPaths::ProjectLogDir() / TEXT("WorldPartition") / FString::Printf(TEXT("HLODStats-%08x-%s.csv"), FPlatformProcess::GetCurrentProcessId(), *FDateTime::Now().ToString());
+
+		for (const FWorldContext& Context : GEngine->GetWorldContexts())
+		{
+			if (UWorld* World = Context.World())
+			{
+				if (UWorldPartitionHLODEditorSubsystem* HLODEditorSubsystem = World->GetSubsystem<UWorldPartitionHLODEditorSubsystem>())
+				{
+					IWorldPartitionEditorModule::FWriteHLODStatsParams Params;
+					Params.World = World;
+					Params.StatsType = IWorldPartitionEditorModule::FWriteHLODStatsParams::EStatsType::Default;
+					Params.Filename = HLODStatsOutputFilename;
+					HLODEditorSubsystem->WriteHLODStats(Params);
+				}
+			}
+		}
+	})
+);
+
+FAutoConsoleCommand HLODDumpInputStats(
+	TEXT("wp.Editor.HLOD.DumpInputStats"),
+	TEXT("Export stats regarding the input to HLOD generation to a CSV formatted file."),
+	FConsoleCommandWithArgsDelegate::CreateLambda([](const TArray<FString>& Args)
+	{
+		const FString HLODStatsOutputFilename = FPaths::ProjectLogDir() / TEXT("WorldPartition") / FString::Printf(TEXT("HLODInputStats-%08x-%s.csv"), FPlatformProcess::GetCurrentProcessId(), *FDateTime::Now().ToString());
+
+		for (const FWorldContext& Context : GEngine->GetWorldContexts())
+		{
+			if (UWorld* World = Context.World())
+			{
+				if (UWorldPartitionHLODEditorSubsystem* HLODEditorSubsystem = World->GetSubsystem<UWorldPartitionHLODEditorSubsystem>())
+				{
+					IWorldPartitionEditorModule::FWriteHLODStatsParams Params;
+					Params.World = World;
+					Params.StatsType = IWorldPartitionEditorModule::FWriteHLODStatsParams::EStatsType::InputDetails;
+					Params.Filename = HLODStatsOutputFilename;
+					HLODEditorSubsystem->WriteHLODStats(Params);
+				}
+			}
+		}
+	})
+);
+
 #undef LOCTEXT_NAMESPACE
diff --git a/Engine/Source/Editor/WorldPartitionEditor/Private/WorldPartitionEditorModule.cpp b/Engine/Source/Editor/WorldPartitionEditor/Private/WorldPartitionEditorModule.cpp
--- a/Engine/Source/Editor/WorldPartitionEditor/Private/WorldPartitionEditorModule.cpp
+++ b/Engine/Source/Editor/WorldPartitionEditor/Private/WorldPartitionEditorModule.cpp
@@ -20,6 +20,7 @@
 #include "WorldPartition/WorldPartitionEditorSettings.h"
 #include "WorldPartition/HLOD/HLODActor.h"
 #include "WorldPartition/HLOD/SWorldPartitionBuildHLODsDialog.h"
+#include "WorldPartition/HLOD/HLODEditorSubsystem.h"
 #include "WorldPartition/WorldPartitionClassDescRegistry.h"
 
 #include "LevelEditor.h"
@@ -475,6 +476,16 @@ bool FWorldPartitionEditorModule::IsHLODInEditorAllowed(UWorld* InWorld, FText*
 	return true;
 }
 
+bool FWorldPartitionEditorModule::WriteHLODStats(const FWriteHLODStatsParams& Params) const
+{
+	UWorldPartitionHLODEditorSubsystem* HLODEditorSubsystem = Params.World->GetSubsystem<UWorldPartitionHLODEditorSubsystem>();
+	if (ensure(HLODEditorSubsystem))
+	{
+		return HLODEditorSubsystem->WriteHLODStats(Params);
+	}
+	return false;
+}
+
 void FWorldPartitionEditorModule::OnConvertMap()
 {
 	IContentBrowserSingleton& ContentBrowserSingleton = FModuleManager::LoadModuleChecked<FContentBrowserModule>("ContentBrowser").Get();
diff --git a/Engine/Source/Editor/WorldPartitionEditor/Public/WorldPartition/HLOD/HLODEditorSubsystem.h b/Engine/Source/Editor/WorldPartitionEditor/Public/WorldPartition/HLOD/HLODEditorSubsystem.h
--- a/Engine/Source/Editor/WorldPartitionEditor/Public/WorldPartition/HLOD/HLODEditorSubsystem.h
+++ b/Engine/Source/Editor/WorldPartitionEditor/Public/WorldPartition/HLOD/HLODEditorSubsystem.h
@@ -3,6 +3,7 @@
 #pragma once
 
 #include "Subsystems/WorldSubsystem.h"
+#include "WorldPartition/IWorldPartitionEditorModule.h"
 #include "WorldPartition/WorldPartitionActorLoaderInterface.h"
 #include "WorldPartition/WorldPartitionHandle.h"
 #include "Templates/PimplPtr.h"
@@ -54,6 +55,8 @@ public:
 	virtual TStatId GetStatId() const override;
 	//~ End FTickableGameObject Interface
 
+	virtual bool WriteHLODStats(const IWorldPartitionEditorModule::FWriteHLODStatsParams& Params) const;
+
 	static void AddHLODSettingsFilter(EHLODSettingsVisibility InSettingsVisibility, TSoftObjectPtr<UStruct> InStruct, FName InPropertyName);
 	
 private:
@@ -67,6 +70,9 @@ private:
 
 	void ForceHLODStateUpdate();
 
+	bool WriteHLODStats(const FString& InFilename) const;
+	bool WriteHLODInputStats(const FString& InFilename) const;
+	
 	void OnWorldPartitionEditorSettingsChanged(const FName& PropertyName, const UWorldPartitionEditorSettings& WorldPartitionEditorSettings);
 	void ApplyHLODSettingsFiltering();
 
diff --git a/Engine/Source/Editor/WorldPartitionEditor/Public/WorldPartitionEditorModule.h b/Engine/Source/Editor/WorldPartitionEditor/Public/WorldPartitionEditorModule.h
--- a/Engine/Source/Editor/WorldPartitionEditor/Public/WorldPartitionEditorModule.h
+++ b/Engine/Source/Editor/WorldPartitionEditor/Public/WorldPartitionEditorModule.h
@@ -91,6 +91,11 @@ public:
 
 	virtual bool IsHLODInEditorAllowed(UWorld* InWorld, FText* OutDisallowedReason) const override;
 
+	/**
+     * Writes various HLOD stats to a file.
+     */ 
+	virtual bool WriteHLODStats(const FWriteHLODStatsParams& Params) const override;
+
 	/**
 	 * Convert the specified map to a world partition map.
 	 */
diff --git a/Engine/Source/Runtime/Engine/Private/UnrealEngine.cpp b/Engine/Source/Runtime/Engine/Private/UnrealEngine.cpp
--- a/Engine/Source/Runtime/Engine/Private/UnrealEngine.cpp
+++ b/Engine/Source/Runtime/Engine/Private/UnrealEngine.cpp
@@ -12419,7 +12419,6 @@ float DrawMapWarnings(UWorld* World, FViewport* Viewport, FCanvas* Canvas, UCanv
 		}
 	}
 
-#if WITH_EDITOR
 	if (UWorldPartitionHLODRuntimeSubsystem* HLODSubsystem = World->GetSubsystem<UWorldPartitionHLODRuntimeSubsystem>())
 	{
 		uint32 NumOutdatedHLODActors = HLODSubsystem->GetNumOutdatedHLODActors();
@@ -12431,8 +12430,6 @@ float DrawMapWarnings(UWorld* World, FViewport* Viewport, FCanvas* Canvas, UCanv
 			MessageY += FontSizeY;
 		}
 	}
-#endif
-
 #endif
 
 	if (World->NumTextureStreamingUnbuiltComponents > 0 || World->NumTextureStreamingDirtyResources > 0)
diff --git a/Engine/Source/Runtime/Engine/Private/WorldPartition/HLOD/HLODActor.cpp b/Engine/Source/Runtime/Engine/Private/WorldPartition/HLOD/HLODActor.cpp
--- a/Engine/Source/Runtime/Engine/Private/WorldPartition/HLOD/HLODActor.cpp
+++ b/Engine/Source/Runtime/Engine/Private/WorldPartition/HLOD/HLODActor.cpp
@@ -430,6 +430,7 @@ void AWorldPartitionHLOD::SetHLODComponents(const TArray<UActorComponent*>& InHL
 void AWorldPartitionHLOD::SetSourceActors(UWorldPartitionHLODSourceActors* InHLODSourceActors)
 {
 	SourceActors = InHLODSourceActors;
+	InputStats.BuildersReferencedAssets.Reset();
 }
 
 const UWorldPartitionHLODSourceActors* AWorldPartitionHLOD::GetSourceActors() const
@@ -442,6 +443,16 @@ UWorldPartitionHLODSourceActors* AWorldPartitionHLOD::GetSourceActors()
 	return SourceActors;
 }
 
+void AWorldPartitionHLOD::SetInputStats(const FHLODBuildInputStats& InInputStats)
+{
+	InputStats = InInputStats;
+}
+
+const FHLODBuildInputStats& AWorldPartitionHLOD::GetInputStats() const
+{
+	return InputStats;
+}
+
 void AWorldPartitionHLOD::SetSourceCellGuid(const FGuid& InSourceCellGuid)
 {
 	SourceCellGuid = InSourceCellGuid;
diff --git a/Engine/Source/Runtime/Engine/Private/WorldPartition/HLOD/HLODBuilder.cpp b/Engine/Source/Runtime/Engine/Private/WorldPartition/HLOD/HLODBuilder.cpp
--- a/Engine/Source/Runtime/Engine/Private/WorldPartition/HLOD/HLODBuilder.cpp
+++ b/Engine/Source/Runtime/Engine/Private/WorldPartition/HLOD/HLODBuilder.cpp
@@ -14,6 +14,9 @@
 #include "UObject/Package.h"
 #include "WorldPartition/HLOD/HLODInstancedStaticMeshComponent.h"
 
+#include "AssetRegistry/AssetData.h"
+#include "AssetRegistry/AssetRegistryModule.h"
+
 #include UE_INLINE_GENERATED_CPP_BY_NAME(HLODBuilder)
 
 
@@ -246,7 +249,7 @@ static bool ShouldBatchComponent(UActorComponent* ActorComponent)
 	return bShouldBatch;
 }
 
-TArray<UActorComponent*> UHLODBuilder::Build(const FHLODBuildContext& InHLODBuildContext) const
+FHLODBuildResult UHLODBuilder::Build(const FHLODBuildContext& InHLODBuildContext) const
 {
 	// Handle components using a batching policy separately
 	TArray<UActorComponent*> InputComponents;
@@ -279,7 +282,36 @@ TArray<UActorComponent*> UHLODBuilder::Build(const FHLODBuildContext& InHLODBuil
 		TSubclassOf<UHLODBuilder> HLODBuilderClass = SourceComponent->GetCustomHLODBuilderClass();
 		HLODBuildersForComponents.FindOrAdd(HLODBuilderClass).Add(SourceComponent);
 	}
+	
+	FHLODBuildResult BuildResult;
+
+	auto AddReferencedAssetsToStats = [&BuildResult](FName HLODBuilderClassName, TArray<UActorComponent*> InSourceComponents)
+	{
+		FAssetRegistryModule& AssetRegistryModule = FModuleManager::GetModuleChecked<FAssetRegistryModule>("AssetRegistry");
+		IAssetRegistry& AssetRegistry = AssetRegistryModule.Get();
+
+		const FTopLevelAssetPath StaticMeshAssetClassPath(UStaticMesh::StaticClass());
 
+		FHLODBuildInputReferencedAssets& ReferencedAssetsStats = BuildResult.InputStats.BuildersReferencedAssets.FindOrAdd(HLODBuilderClassName);
+
+		for (UActorComponent* SourceComponent : InSourceComponents)
+		{
+			// At the moment we only care about static meshes for our stats
+			if (UStaticMeshComponent* StaticMeshComponent = Cast<UStaticMeshComponent>(SourceComponent))
+			{
+				FAssetData AssetData = AssetRegistry.GetAssetByObjectPath(FSoftObjectPath(StaticMeshComponent->GetStaticMesh()));
+				if (AssetData.IsUAsset())
+				{					
+					if (AssetData.AssetClassPath == StaticMeshAssetClassPath)
+					{
+						FTopLevelAssetPath StaticMeshAssetPath(AssetData.PackageName, AssetData.AssetName);
+						ReferencedAssetsStats.StaticMeshes.FindOrAdd(StaticMeshAssetPath)++;
+					}
+				}
+			}
+		}	
+	};
+	
 	// Build HLOD components by sending source components to the individual builders, in batch
 	TArray<UActorComponent*> HLODComponents;
 	for (const auto& HLODBuilderPair : HLODBuildersForComponents)
@@ -288,17 +320,24 @@ TArray<UActorComponent*> UHLODBuilder::Build(const FHLODBuildContext& InHLODBuil
 		const UHLODBuilder* HLODBuilder = HLODBuilderPair.Key ? HLODBuilderPair.Key->GetDefaultObject<UHLODBuilder>() : this;
 		const TArray<UActorComponent*>& SourceComponents = HLODBuilderPair.Value;
 
+		AddReferencedAssetsToStats(HLODBuilder->GetClass()->GetFName(), SourceComponents);
+
 		TArray<UActorComponent*> NewComponents = HLODBuilder->Build(InHLODBuildContext, SourceComponents);
-		HLODComponents.Append(NewComponents);
+		BuildResult.HLODComponents.Append(NewComponents);
 	}
 
 	// Append batched components
-	HLODComponents.Append(BatchInstances(ComponentsToBatch));
+	if (!ComponentsToBatch.IsEmpty())
+	{
+		FName HLODBuilderInstancingClassName("HLODBuilderInstancing");
+		AddReferencedAssetsToStats(HLODBuilderInstancingClassName, ComponentsToBatch);
+		BuildResult.HLODComponents.Append(BatchInstances(ComponentsToBatch));
+	}
 
 	// In case a builder returned null entries, clean the array.
-	HLODComponents.RemoveSwap(nullptr);
+	BuildResult.HLODComponents.RemoveSwap(nullptr);
 
-	return HLODComponents;
+	return BuildResult;
 }
 
 #endif // WITH_EDITOR
diff --git a/Engine/Source/Runtime/Engine/Private/WorldPartition/HLOD/HLODRuntimeSubsystem.cpp b/Engine/Source/Runtime/Engine/Private/WorldPartition/HLOD/HLODRuntimeSubsystem.cpp
--- a/Engine/Source/Runtime/Engine/Private/WorldPartition/HLOD/HLODRuntimeSubsystem.cpp
+++ b/Engine/Source/Runtime/Engine/Private/WorldPartition/HLOD/HLODRuntimeSubsystem.cpp
@@ -5,9 +5,7 @@
 #include "WorldPartition/HLOD/HLODActor.h"
 #include "Engine/LevelStreaming.h"
 #include "Engine/World.h"
-#include "WorldPartition/HLOD/HLODActorDesc.h"
 #include "Materials/MaterialInterface.h"
-#include "WorldPartition/DataLayer/DataLayerManager.h"
 #include "Misc/Paths.h"
 #include "WorldPartition/WorldPartition.h"
 #include "RenderUtils.h"
@@ -21,14 +19,9 @@
 #include "EngineModule.h"
 #include "LevelUtils.h"
 #include "Components/InstancedStaticMeshComponent.h"
-#include "Misc/FileHelper.h"
 #include "SceneViewExtension.h"
 #include "StaticMeshResources.h"
 #include "Subsystems/Subsystem.h"
-#include "UObject/UObjectIterator.h"
-#include "WorldPartition/DataLayer/WorldDataLayers.h"
-#include "WorldPartition/HLOD/HLODStats.h"
-#include "WorldPartition/WorldPartitionActorDescInstance.h"
 
 #include UE_INLINE_GENERATED_CPP_BY_NAME(HLODRuntimeSubsystem)
 
@@ -342,7 +335,7 @@ void UWorldPartitionHLODRuntimeSubsystem::RegisterHLODActor(AWorldPartitionHLOD*
 		UE_LOG(LogHLODRuntimeSubsystem, Verbose, TEXT("Found HLOD %s referencing nonexistent cell '%s'"), *InWorldPartitionHLOD->GetActorNameOrLabel(), *InWorldPartitionHLOD->GetSourceCellGuid().ToString());
 		InWorldPartitionHLOD->SetVisibility(false);
 
-#if WITH_EDITOR
+#if !UE_BUILD_SHIPPING
 		OutdatedHLODActors.Add(InWorldPartitionHLOD);
 #endif
 	}
@@ -363,7 +356,7 @@ void UWorldPartitionHLODRuntimeSubsystem::UnregisterHLODActor(AWorldPartitionHLO
 	}
 	else
 	{
-#if WITH_EDITOR
+#if !UE_BUILD_SHIPPING
 		OutdatedHLODActors.Remove(InWorldPartitionHLOD);
 #endif
 	}
@@ -725,111 +718,5 @@ void UWorldPartitionHLODRuntimeSubsystem::OnBeginRenderViews(const FSceneViewFam
 	MakeHLODRenderResourcesResident(VTRequests, NaniteRequests, InViewFamily);
 }
 
-#if WITH_EDITOR
-
-bool UWorldPartitionHLODRuntimeSubsystem::WriteHLODStatsCSV(UWorld* InWorld, const FString& InFilename)
-{
-	UWorldPartition* WorldPartition = InWorld ? InWorld->GetWorldPartition() : nullptr;
-	if (!WorldPartition)
-	{
-		return false;
-	}
-	
-	typedef TFunction<FString(FWorldPartitionActorDescInstance* InActorDescInstance, const FHLODActorDesc&)> FGetStatFunc;
-
-	auto GetHLODStat = [](FName InStatName)
-	{
-		return TPair<FName, FGetStatFunc>(InStatName, [InStatName](FWorldPartitionActorDescInstance* InActorDescInstance, const FHLODActorDesc& InActorDesc)
-		{
-			return FString::Printf(TEXT("%lld"), InActorDesc.GetStat(InStatName));
-		});
-	};
-
-	const UDataLayerManager* DataLayerManager = WorldPartition->GetDataLayerManager();
-
-	auto GetDataLayerShortName = [DataLayerManager](FName DataLayerInstanceName)
-	{ 
-		const UDataLayerInstance* DataLayerInstance = DataLayerManager ? DataLayerManager->GetDataLayerInstance(DataLayerInstanceName) : nullptr;
-		return DataLayerInstance ? DataLayerInstance->GetDataLayerShortName() : DataLayerInstanceName.ToString();
-	};
-
-	TArray<TPair<FName, FGetStatFunc>> StatsToWrite =
-	{	
-		{ "WorldPackage",		[InWorld](FWorldPartitionActorDescInstance* InActorDescInstance, const FHLODActorDesc& InActorDesc) { return InWorld->GetPackage()->GetName(); } },
-		{ "Name",				[](FWorldPartitionActorDescInstance* InActorDescInstance, const FHLODActorDesc& InActorDesc) { return InActorDescInstance->GetActorLabelString(); } },
-		{ "HLODLayer",			[](FWorldPartitionActorDescInstance* InActorDescInstance, const FHLODActorDesc& InActorDesc) { return InActorDesc.GetSourceHLODLayer().GetAssetName().ToString(); }},
-		{ "SpatiallyLoaded",	[](FWorldPartitionActorDescInstance* InActorDescInstance, const FHLODActorDesc& InActorDesc) { return InActorDescInstance->GetIsSpatiallyLoaded() ? TEXT("true") : TEXT("false"); } },
-		{ "DataLayers",			[&GetDataLayerShortName](FWorldPartitionActorDescInstance* InActorDescInstance, const FHLODActorDesc& InActorDesc) { return FString::JoinBy(InActorDescInstance->GetDataLayerInstanceNames().ToArray(), TEXT(" | "), GetDataLayerShortName); }},
-
-		GetHLODStat(FWorldPartitionHLODStats::InputActorCount),
-		GetHLODStat(FWorldPartitionHLODStats::InputTriangleCount),
-		GetHLODStat(FWorldPartitionHLODStats::InputVertexCount),
-
-		GetHLODStat(FWorldPartitionHLODStats::MeshInstanceCount),
-		GetHLODStat(FWorldPartitionHLODStats::MeshNaniteTriangleCount),
-		GetHLODStat(FWorldPartitionHLODStats::MeshNaniteVertexCount),
-		GetHLODStat(FWorldPartitionHLODStats::MeshTriangleCount),
-		GetHLODStat(FWorldPartitionHLODStats::MeshVertexCount),
-		GetHLODStat(FWorldPartitionHLODStats::MeshUVChannelCount),
-
-		GetHLODStat(FWorldPartitionHLODStats::MaterialBaseColorTextureSize),
-		GetHLODStat(FWorldPartitionHLODStats::MaterialNormalTextureSize),
-		GetHLODStat(FWorldPartitionHLODStats::MaterialEmissiveTextureSize),
-		GetHLODStat(FWorldPartitionHLODStats::MaterialMetallicTextureSize),
-		GetHLODStat(FWorldPartitionHLODStats::MaterialRoughnessTextureSize),
-		GetHLODStat(FWorldPartitionHLODStats::MaterialSpecularTextureSize),
-		
-		GetHLODStat(FWorldPartitionHLODStats::MemoryMeshResourceSizeBytes),
-		GetHLODStat(FWorldPartitionHLODStats::MemoryTexturesResourceSizeBytes),
-		GetHLODStat(FWorldPartitionHLODStats::MemoryDiskSizeBytes),
-		
-		GetHLODStat(FWorldPartitionHLODStats::BuildTimeLoadMilliseconds),
-		GetHLODStat(FWorldPartitionHLODStats::BuildTimeBuildMilliseconds),
-		GetHLODStat(FWorldPartitionHLODStats::BuildTimeTotalMilliseconds)				
-	};
-
-	FStringOutputDevice Output;
-
-	// Write header if file doesn't exist
-	if (!IFileManager::Get().FileExists(*InFilename))
-	{
-		const FString StatHeader = FString::JoinBy(StatsToWrite, TEXT(","), [](const TPair<FName, FGetStatFunc>& Pair) { return Pair.Key.ToString(); });
-		Output.Logf(TEXT("%s" LINE_TERMINATOR_ANSI), *StatHeader);
-	}
-
-	// Write one line per HLOD actor desc
-	for (FActorDescContainerInstanceCollection::TIterator<AWorldPartitionHLOD> HLODIterator(WorldPartition); HLODIterator; ++HLODIterator)
-	{
-		const FString StatLine = FString::JoinBy(StatsToWrite, TEXT(","), [&HLODIterator](const TPair<FName, FGetStatFunc>& Pair) 
-		{ 
-			const FHLODActorDesc& HLODActorDesc = *(FHLODActorDesc*)HLODIterator->GetActorDesc();
-			return Pair.Value(*HLODIterator, HLODActorDesc); 
-		});
-		Output.Logf(TEXT("%s" LINE_TERMINATOR_ANSI), *StatLine);
-	}
-
-	// Write to file
-	return FFileHelper::SaveStringToFile(Output, *InFilename, FFileHelper::EEncodingOptions::AutoDetect, &IFileManager::Get(), EFileWrite::FILEWRITE_Append);
-}
-
-FAutoConsoleCommand HLODDumpStats(
-	TEXT("wp.Editor.HLOD.DumpStats"),
-	TEXT("Write various HLOD stats to a CSV formatted file."),
-	FConsoleCommandWithArgsDelegate::CreateLambda([](const TArray<FString>& Args)
-	{
-		const FString HLODStatsOutputFilename = FPaths::ProjectLogDir() / TEXT("WorldPartition") / FString::Printf(TEXT("HLODStats-%08x-%s.csv"), FPlatformProcess::GetCurrentProcessId(), *FDateTime::Now().ToString());
-
-		for (const FWorldContext& Context : GEngine->GetWorldContexts())
-		{
-			if (UWorld* World = Context.World())
-			{
-				UWorldPartitionHLODRuntimeSubsystem::WriteHLODStatsCSV(World, HLODStatsOutputFilename);
-			}
-		}
-	})
-);
-
-#endif // #if WITH_EDITOR
-
 #undef LOCTEXT_NAMESPACE
 
diff --git a/Engine/Source/Runtime/Engine/Public/WorldPartition/HLOD/HLODActor.h b/Engine/Source/Runtime/Engine/Public/WorldPartition/HLOD/HLODActor.h
--- a/Engine/Source/Runtime/Engine/Public/WorldPartition/HLOD/HLODActor.h
+++ b/Engine/Source/Runtime/Engine/Public/WorldPartition/HLOD/HLODActor.h
@@ -8,6 +8,7 @@
 #include "WorldPartition/WorldPartitionRuntimeCell.h"
 #include "WorldPartition/WorldPartitionActorDesc.h"
 #include "WorldPartition/HLOD/HLODStats.h"
+#include "WorldPartition/HLOD/HLODBuilder.h"
 
 #if WITH_EDITOR
 #include "UObject/ObjectSaveContext.h"
@@ -45,6 +46,9 @@ public:
 	ENGINE_API UWorldPartitionHLODSourceActors* GetSourceActors();
 	ENGINE_API const UWorldPartitionHLODSourceActors* GetSourceActors() const;
 
+	ENGINE_API void SetInputStats(const FHLODBuildInputStats& InInputStats);
+	ENGINE_API const FHLODBuildInputStats& GetInputStats() const;
+
 	void SetRequireWarmup(bool InRequireWarmup) { bRequireWarmup = InRequireWarmup; }
 
 	ENGINE_API void SetSourceCellGuid(const FGuid& InSourceCellGuid);
@@ -106,6 +110,9 @@ private:
 	UPROPERTY()
 	TObjectPtr<UWorldPartitionHLODSourceActors> SourceActors;
 
+	UPROPERTY()
+	FHLODBuildInputStats InputStats;
+
 	UPROPERTY()
 	FBox HLODBounds;
 
diff --git a/Engine/Source/Runtime/Engine/Public/WorldPartition/HLOD/HLODActorDesc.h b/Engine/Source/Runtime/Engine/Public/WorldPartition/HLOD/HLODActorDesc.h
--- a/Engine/Source/Runtime/Engine/Public/WorldPartition/HLOD/HLODActorDesc.h
+++ b/Engine/Source/Runtime/Engine/Public/WorldPartition/HLOD/HLODActorDesc.h
@@ -28,7 +28,7 @@ public:
 	inline const TArray<FGuid>& GetChildHLODActors() const { return ChildHLODActors; }
 	inline const FTopLevelAssetPath& GetSourceHLODLayer() const { return SourceHLODLayer; }
 	inline const FStats& GetStats() const { return HLODStats; }
-	int64 GetStat(FName InStatName) const;
+	ENGINE_API int64 GetStat(FName InStatName) const;
 
 	//~ Begin FWorldPartitionActorDesc Interface.
 	virtual FBox GetEditorBounds() const override { return EditorBounds; }
diff --git a/Engine/Source/Runtime/Engine/Public/WorldPartition/HLOD/HLODBuilder.h b/Engine/Source/Runtime/Engine/Public/WorldPartition/HLOD/HLODBuilder.h
--- a/Engine/Source/Runtime/Engine/Public/WorldPartition/HLOD/HLODBuilder.h
+++ b/Engine/Source/Runtime/Engine/Public/WorldPartition/HLOD/HLODBuilder.h
@@ -57,6 +57,38 @@ struct FHLODBuildContext
 	double MinVisibleDistance;
 };
 
+/**
+ * Keep track of assets used as input to the HLOD generation of a given HLOD actor, along with the number of occurences.
+ */
+USTRUCT()
+struct FHLODBuildInputReferencedAssets
+{
+	GENERATED_BODY()
+
+	UPROPERTY()
+	TMap<FTopLevelAssetPath, uint32> StaticMeshes;
+};
+
+/**
+ * Referenced assets per HLOD builder.
+ */
+USTRUCT()
+struct FHLODBuildInputStats
+{
+	GENERATED_BODY()
+
+	UPROPERTY()
+	TMap<FName, FHLODBuildInputReferencedAssets> BuildersReferencedAssets;
+};
+
+/**
+ * Result of the HLOD build of a single actor.
+ */
+struct FHLODBuildResult
+{
+	FHLODBuildInputStats		InputStats;
+	TArray<UActorComponent*>	HLODComponents;
+};
 
 /**
  * Base class for all HLOD builders
@@ -78,10 +110,10 @@ public:
 	 * Build an HLOD representation of the input actors.
 	 * Components returned by this method needs to be properly outered & assigned to your target (HLOD) actor.
 	 */
-	ENGINE_API TArray<UActorComponent*> Build(const FHLODBuildContext& InHLODBuildContext) const;
+	ENGINE_API FHLODBuildResult Build(const FHLODBuildContext& InHLODBuildContext) const;
 
 	UE_DEPRECATED(5.2, "Use Build() method that takes a single FHLODBuildContext parameter.")
-	TArray<UActorComponent*> Build(const FHLODBuildContext& InHLODBuildContext, const TArray<AActor*>& InSourceActors) const { return Build(InHLODBuildContext); }
+	ENGINE_API TArray<UActorComponent*> Build(const FHLODBuildContext& InHLODBuildContext, const TArray<AActor*>& InSourceActors) const { return Build(InHLODBuildContext).HLODComponents; }
 
 	/**
 	 * Return the setting subclass associated with this builder.
diff --git a/Engine/Source/Runtime/Engine/Public/WorldPartition/HLOD/HLODRuntimeSubsystem.h b/Engine/Source/Runtime/Engine/Public/WorldPartition/HLOD/HLODRuntimeSubsystem.h
--- a/Engine/Source/Runtime/Engine/Public/WorldPartition/HLOD/HLODRuntimeSubsystem.h
+++ b/Engine/Source/Runtime/Engine/Public/WorldPartition/HLOD/HLODRuntimeSubsystem.h
@@ -62,9 +62,13 @@ public:
 
 	ENGINE_API void OnCVarsChanged();
 
-#if WITH_EDITOR
+#if !UE_BUILD_SHIPPING
 	uint32 GetNumOutdatedHLODActors() const { return OutdatedHLODActors.Num(); }
-	static ENGINE_API bool WriteHLODStatsCSV(UWorld* InWorld, const FString& InFilename);
+#endif
+
+#if WITH_EDITOR
+	UE_DEPRECATED(5.5, "Use UWorldPartitionHLODEditorSubsystem::WriteHLODStatsCSV()")
+	static ENGINE_API bool WriteHLODStatsCSV(UWorld* InWorld, const FString& InFilename) { return false; }
 #endif
 
 	UE_DEPRECATED(5.4, "You should perform this logic on the game side.")
@@ -117,7 +121,7 @@ private:
 	FWorldPartitionHLODActorRegisteredEvent		HLODActorRegisteredEvent;
 	FWorldPartitionHLODActorUnregisteredEvent	HLODActorUnregisteredEvent;
 
-#if WITH_EDITOR
+#if !UE_BUILD_SHIPPING
 	TSet<AWorldPartitionHLOD*> OutdatedHLODActors;
 #endif
 };
diff --git a/Engine/Source/Runtime/Engine/Public/WorldPartition/IWorldPartitionEditorModule.h b/Engine/Source/Runtime/Engine/Public/WorldPartition/IWorldPartitionEditorModule.h
--- a/Engine/Source/Runtime/Engine/Public/WorldPartition/IWorldPartitionEditorModule.h
+++ b/Engine/Source/Runtime/Engine/Public/WorldPartition/IWorldPartitionEditorModule.h
@@ -69,6 +69,25 @@ public:
 
 	virtual bool IsHLODInEditorAllowed(UWorld* InWorld, FText* OutDisallowedReason = nullptr) const = 0;
 
+	/**
+     * Parameters for the WriteHLODStats() method.
+     */	
+	struct FWriteHLODStatsParams
+	{
+		enum class EStatsType : uint8
+		{
+			Default,
+			InputDetails
+		};
+
+		UWorld* World = nullptr;
+		FString Filename;
+		EStatsType StatsType = EStatsType::Default;
+	};
+
+	/** Writes various HLOD stats to a file. */
+	virtual bool WriteHLODStats(const FWriteHLODStatsParams& Params) const = 0;
+
 	/** Triggered when a world is added. */
 	DECLARE_EVENT_OneParam(IWorldPartitionEditorModule, FWorldPartitionCreated, UWorld*);
 
