Content Bundle: implement actor colorization mode.

diff --git a/Engine/Source/Editor/WorldPartitionEditor/Private/WorldPartition/ContentBundle/ContentBundleEditorSubsystem.cpp b/Engine/Source/Editor/WorldPartitionEditor/Private/WorldPartition/ContentBundle/ContentBundleEditorSubsystem.cpp
--- a/Engine/Source/Editor/WorldPartitionEditor/Private/WorldPartition/ContentBundle/ContentBundleEditorSubsystem.cpp
+++ b/Engine/Source/Editor/WorldPartitionEditor/Private/WorldPartition/ContentBundle/ContentBundleEditorSubsystem.cpp
@@ -10,6 +10,7 @@
 #include "WorldPartition/DataLayer/ExternalDataLayerInstance.h"
 #include "WorldPartition/DataLayer/ExternalDataLayerAsset.h"
 #include "WorldPartition/WorldPartition.h"
+#include "GameFramework/ActorPrimitiveColorHandler.h"
 #include "Subsystems/ActorEditorContextSubsystem.h"
 #include "DataLayer/DataLayerEditorSubsystem.h"
 #include "Engine/Selection.h"
@@ -19,12 +20,32 @@
 #include "Widgets/Text/STextBlock.h"
 #include "Editor.h"
 
+#define LOCTEXT_NAMESPACE "ContentBundle"
+
+#if ENABLE_ACTOR_PRIMITIVE_COLOR_HANDLER
+static FName NAME_ContentBundleColor(TEXT("ContentBundleColor"));
+#endif
+
 void UContentBundleEditingSubmodule::DoInitialize()
 {
 	check(GEditor);
 	UActorEditorContextSubsystem::Get()->RegisterClient(this);
 	// For backward compatibility, Content Bundle will be applied right when LevelActorAdded is called
 	GEditor->OnLevelActorAdded().AddUObject(this, &UContentBundleEditingSubmodule::ApplyContext);
+
+#if ENABLE_ACTOR_PRIMITIVE_COLOR_HANDLER
+	FActorPrimitiveColorHandler::Get().RegisterPrimitiveColorHandler(NAME_ContentBundleColor, LOCTEXT("ContentBundleColor", "Content Bundle Color"), [](const UPrimitiveComponent* InPrimitiveComponent) -> FLinearColor
+	{
+		if (AActor* Actor = InPrimitiveComponent->GetOwner())
+		{
+			if (Actor->GetContentBundleGuid().IsValid())
+			{
+				return FLinearColor::MakeRandomSeededColor(GetTypeHash(Actor->GetContentBundleGuid()));
+			}
+		}
+		return FLinearColor::Gray;
+	});
+#endif
 }
 
 void UContentBundleEditingSubmodule::DoDenitialize()
@@ -42,6 +63,10 @@ void UContentBundleEditingSubmodule::DoDenitialize()
 	}
 
 	EditingContentBundlesStack.Empty();
+
+#if ENABLE_ACTOR_PRIMITIVE_COLOR_HANDLER
+	FActorPrimitiveColorHandler::Get().UnregisterPrimitiveColorHandler(NAME_ContentBundleColor);
+#endif
 }
 
 void UContentBundleEditingSubmodule::PreEditUndo()
@@ -186,7 +211,6 @@ TSharedRef<SWidget> UContentBundleEditingSubmodule::GetActorEditorContextWidget(
 		];
 	}
 
-
 	return OutWidget;
 }
 
@@ -462,4 +486,5 @@ void UContentBundleEditorSubsystem::SelectActorsInternal(FContentBundleEditor& E
 	GEditor->GetSelectedActors()->EndBatchSelectOperation();
 
 	GEditor->NoteSelectionChange();
-}
\ No newline at end of file
+}
+#undef LOCTEXT_NAMESPACE
\ No newline at end of file
diff --git a/Engine/Source/Runtime/Engine/Private/WorldPartition/WorldPartitionLevelStreamingDynamic.cpp b/Engine/Source/Runtime/Engine/Private/WorldPartition/WorldPartitionLevelStreamingDynamic.cpp
--- a/Engine/Source/Runtime/Engine/Private/WorldPartition/WorldPartitionLevelStreamingDynamic.cpp
+++ b/Engine/Source/Runtime/Engine/Private/WorldPartition/WorldPartitionLevelStreamingDynamic.cpp
@@ -62,13 +62,15 @@ void UWorldPartitionLevelStreamingDynamic::Initialize(const UWorldPartitionRunti
 #if WITH_EDITOR
 	check(ChildPackages.Num() == 0);
 
-	UnsavedActorsContainer = InCell.UnsavedActorsContainer;
+	UnsavedActorsContainer = InCell.UnsavedActorsContainer; 
 
 	Initialize(CellOuterWorld, InCell.GetPackages());
 #else
 	OuterWorldPartition = CellOuterWorld->GetWorldPartition();
 #endif
 
+	LevelColor = InCell.GetCellDebugColor();
+
 	UpdateShouldSkipMakingVisibilityTransactionRequest();
 }
 
@@ -307,11 +309,6 @@ bool UWorldPartitionLevelStreamingDynamic::RequestLevel(UWorld* InPersistentWorl
 			CreateRuntimeLevel();
 			check(RuntimeLevel);
 
-			if (const UWorldPartitionRuntimeLevelStreamingCell* RuntimeLevelStreamingCell = StreamingCell.Get())
-			{
-				LevelColor = RuntimeLevelStreamingCell->GetCellDebugColor();
-			}
-
 			UPackage* CellLevelPackage = RuntimeLevel->GetPackage();
 			check(CellLevelPackage);
 			check(UWorld::FindWorldInPackage(CellLevelPackage));
