Update HLOD callsite of FPropertyEditorPermissionList::Get().AddPermissionList to use new version of function that takes owner name
#rnx

diff --git a/Engine/Source/Editor/WorldPartitionEditor/Private/WorldPartition/HLOD/HLODEditorSubsystem.cpp b/Engine/Source/Editor/WorldPartitionEditor/Private/WorldPartition/HLOD/HLODEditorSubsystem.cpp
--- a/Engine/Source/Editor/WorldPartitionEditor/Private/WorldPartition/HLOD/HLODEditorSubsystem.cpp
+++ b/Engine/Source/Editor/WorldPartitionEditor/Private/WorldPartition/HLOD/HLODEditorSubsystem.cpp
@@ -148,7 +148,7 @@ void UWorldPartitionHLODEditorSubsystem::ApplyHLODSettingsFiltering()
 				PermissionList.AddAllowListItem(PropertyPermissionListOwnerName, PropertyName);
 			}
 
-			FPropertyEditorPermissionList::Get().AddPermissionList(StructProperties.Key, PermissionList, EPropertyPermissionListRules::UseExistingPermissionList);
+			FPropertyEditorPermissionList::Get().AddPermissionList(StructProperties.Key, PermissionList, EPropertyPermissionListRules::UseExistingPermissionList, { PropertyPermissionListOwnerName });
 		}
 	}
 }
