World Partition - HLOD: Disable "Build HLOD" buttons on CDO or incomplete HLOD actors

diff --git a/Engine/Source/Editor/WorldPartitionEditor/Private/WorldPartition/Customizations/WorldPartitionHLODDetailsCustomization.cpp b/Engine/Source/Editor/WorldPartitionEditor/Private/WorldPartition/Customizations/WorldPartitionHLODDetailsCustomization.cpp
--- a/Engine/Source/Editor/WorldPartitionEditor/Private/WorldPartition/Customizations/WorldPartitionHLODDetailsCustomization.cpp
+++ b/Engine/Source/Editor/WorldPartitionEditor/Private/WorldPartition/Customizations/WorldPartitionHLODDetailsCustomization.cpp
@@ -39,6 +39,7 @@ void FWorldPartitionHLODDetailsCustomization::CustomizeDetails(IDetailLayoutBuil
 			SNew(SButton)
 			.Text(LOCTEXT("BuildHLODButtonText", "Build HLOD"))
 			.OnClicked(this, &FWorldPartitionHLODDetailsCustomization::OnBuildHLOD)
+			.IsEnabled(this, &FWorldPartitionHLODDetailsCustomization::CanBuildHLOD)
 		]
 	];
 }
@@ -55,6 +56,31 @@ TArray<AWorldPartitionHLOD*> FWorldPartitionHLODDetailsCustomization::GetSelecte
 	return SelectedHLODActors;
 }
 
+bool FWorldPartitionHLODDetailsCustomization::CanBuildHLOD() const
+{
+	TArray<AWorldPartitionHLOD*> SelectedHLODActors = GetSelectedHLODActors();
+
+	for (AWorldPartitionHLOD* HLODActor : SelectedHLODActors)
+	{
+		if (HLODActor == nullptr)
+		{
+			return false;
+		}
+
+		if (HLODActor->IsTemplate())
+		{
+			return false;
+		}
+
+		if (HLODActor->GetSourceActors() == nullptr)
+		{
+			return false;
+		}
+	}
+
+	return true;
+}
+
 FReply FWorldPartitionHLODDetailsCustomization::OnBuildHLOD()
 {
 	TArray<AWorldPartitionHLOD*> SelectedHLODActors = GetSelectedHLODActors();
diff --git a/Engine/Source/Editor/WorldPartitionEditor/Private/WorldPartition/Customizations/WorldPartitionHLODDetailsCustomization.h b/Engine/Source/Editor/WorldPartitionEditor/Private/WorldPartition/Customizations/WorldPartitionHLODDetailsCustomization.h
--- a/Engine/Source/Editor/WorldPartitionEditor/Private/WorldPartition/Customizations/WorldPartitionHLODDetailsCustomization.h
+++ b/Engine/Source/Editor/WorldPartitionEditor/Private/WorldPartition/Customizations/WorldPartitionHLODDetailsCustomization.h
@@ -19,6 +19,7 @@ protected:
 	FWorldPartitionHLODDetailsCustomization()
 	{}
 
+	bool CanBuildHLOD() const;
 	FReply OnBuildHLOD();
 
 private:
