World Partition - Fixed issue where the camera location in the World Partition Editor wouldn't update anymore after changing the viewport type
* UUnrealEditorSubsystem::GetLevelViewportCameraInfo() is actually incorrect for that use case as it will return the wrong viewport.
* HLOD in editor was also affected by this issue
* As a bonus, the camera location will now be updated when working in orthographic projection

diff --git a/Engine/Source/Editor/WorldPartitionEditor/Private/WorldPartition/HLOD/HLODEditorSubsystem.cpp b/Engine/Source/Editor/WorldPartitionEditor/Private/WorldPartition/HLOD/HLODEditorSubsystem.cpp
--- a/Engine/Source/Editor/WorldPartitionEditor/Private/WorldPartition/HLOD/HLODEditorSubsystem.cpp
+++ b/Engine/Source/Editor/WorldPartitionEditor/Private/WorldPartition/HLOD/HLODEditorSubsystem.cpp
@@ -245,17 +245,20 @@ void UWorldPartitionHLODEditorSubsystem::Tick(float DeltaTime)
 			{
 				FVector CameraLocation;
 				FRotator CameraRotation;
-				UnrealEditorSubsystem->GetLevelViewportCameraInfo(CameraLocation, CameraRotation);
+				bool bCameraMoved = false;
 
-				// Camera was moved ?
-				const bool bCameraMoved = CameraLocation != CachedCameraLocation;
-				if (bCameraMoved)
+				if (FWorldPartitionEditorModule::GetActiveLevelViewportCameraInfo(CameraLocation, CameraRotation))
 				{
-					CachedCameraLocation = CameraLocation;
+					// Camera was moved ?
+					bCameraMoved = CameraLocation != CachedCameraLocation;
+					if (bCameraMoved)
+					{
+						CachedCameraLocation = CameraLocation;
+					}
 				}
 
 				if (bForceHLODVisibilityUpdate || bCameraMoved)
-				{					
+				{
 					HLODEditorData->UpdateVisibility(CameraLocation, CachedHLODMinDrawDistance, CachedHLODMaxDrawDistance, bForceHLODVisibilityUpdate);
 				}
 			}
diff --git a/Engine/Source/Editor/WorldPartitionEditor/Private/WorldPartition/SWorldPartitionEditorGrid.cpp b/Engine/Source/Editor/WorldPartitionEditor/Private/WorldPartition/SWorldPartitionEditorGrid.cpp
--- a/Engine/Source/Editor/WorldPartitionEditor/Private/WorldPartition/SWorldPartitionEditorGrid.cpp
+++ b/Engine/Source/Editor/WorldPartitionEditor/Private/WorldPartition/SWorldPartitionEditorGrid.cpp
@@ -9,6 +9,7 @@
 #include "Editor.h"
 #include "Modules/ModuleManager.h"
 #include "Widgets/Text/STextBlock.h"
+#include "WorldPartitionEditorModule.h"
 
 TMap<FName, SWorldPartitionEditorGrid::PartitionEditorGridCreateInstanceFunc> SWorldPartitionEditorGrid::PartitionEditorGridCreateInstanceFactory;
 
@@ -76,14 +77,9 @@ bool SWorldPartitionEditorGrid::GetObserverView(FVector& Location, FRotator& Rot
 	// We are in the editor world
 	if (!GEditor->PlayWorld)
 	{
-		for (const FLevelEditorViewportClient* ViewportClient : GEditor->GetLevelViewportClients())
+		if (FWorldPartitionEditorModule::GetActiveLevelViewportCameraInfo(Location, Rotation))
 		{
-			if (ViewportClient && ViewportClient->IsPerspective())
-			{
-				Rotation = ViewportClient->GetViewRotation();
-				Location = ViewportClient->GetViewLocation();
-				return true;
-			}
+			return true;
 		}
 	}
 
diff --git a/Engine/Source/Editor/WorldPartitionEditor/Private/WorldPartitionEditorModule.cpp b/Engine/Source/Editor/WorldPartitionEditor/Private/WorldPartitionEditorModule.cpp
--- a/Engine/Source/Editor/WorldPartitionEditor/Private/WorldPartitionEditorModule.cpp
+++ b/Engine/Source/Editor/WorldPartitionEditor/Private/WorldPartitionEditorModule.cpp
@@ -27,6 +27,7 @@
 
 #include "LevelEditor.h"
 #include "LevelEditorViewport.h"
+#include "SLevelViewport.h"
 #include "Framework/Application/SlateApplication.h"
 #include "Framework/MultiBox/MultiBoxBuilder.h"
 
@@ -359,6 +360,25 @@ bool FWorldPartitionEditorModule::IsEditingContentBundle(const FGuid& ContentBun
 	return ContentBundleEditorSubsystem && ContentBundleEditorSubsystem->IsEditingContentBundle(ContentBundleGuid);
 }
 
+bool FWorldPartitionEditorModule::GetActiveLevelViewportCameraInfo(FVector& CameraLocation, FRotator& CameraRotation)
+{
+	FLevelEditorModule& LevelEditorModule = FModuleManager::GetModuleChecked<FLevelEditorModule>("LevelEditor");
+	TSharedPtr<ILevelEditor> LevelEditor = LevelEditorModule.GetFirstLevelEditor();
+	if (LevelEditor.IsValid())
+	{
+		TSharedPtr<SLevelViewport> LevelViewport = LevelEditor->GetActiveViewportInterface();
+		if (LevelViewport.IsValid())
+		{
+			const FLevelEditorViewportClient& LevelViewportClient = LevelViewport->GetLevelViewportClient();
+			CameraLocation = LevelViewportClient.GetViewLocation();
+			CameraRotation = LevelViewportClient.GetViewRotation();			
+			return true;
+		}
+	}
+
+	return false;
+}
+
 int32 FWorldPartitionEditorModule::GetPlacementGridSize() const
 {
 	// Currently shares setting with Foliage. Can be changed when exposed.
diff --git a/Engine/Source/Editor/WorldPartitionEditor/Public/WorldPartitionEditorModule.h b/Engine/Source/Editor/WorldPartitionEditor/Public/WorldPartitionEditorModule.h
--- a/Engine/Source/Editor/WorldPartitionEditor/Public/WorldPartitionEditorModule.h
+++ b/Engine/Source/Editor/WorldPartitionEditor/Public/WorldPartitionEditorModule.h
@@ -135,6 +135,11 @@ public:
 	 */
 	bool IsEditingContentBundle(const FGuid& ContentBundleGuid) const;
 
+	/**
+	 * Get the camera location & rotation from the active level viewport.
+	 */
+	static bool GetActiveLevelViewportCameraInfo(FVector& CameraLocation, FRotator& CameraRotation);
+
 private:
 	/** Register menus */
 	void RegisterMenus();
diff --git a/Engine/Source/Editor/WorldPartitionEditor/WorldPartitionEditor.Build.cs b/Engine/Source/Editor/WorldPartitionEditor/WorldPartitionEditor.Build.cs
--- a/Engine/Source/Editor/WorldPartitionEditor/WorldPartitionEditor.Build.cs
+++ b/Engine/Source/Editor/WorldPartitionEditor/WorldPartitionEditor.Build.cs
@@ -16,6 +16,7 @@ public class WorldPartitionEditor : ModuleRules
 				"EditorWidgets",
 				"Engine",
 				"InputCore",
+				"LevelEditor",
 				"Slate",
 				"SlateCore",
 				"UnrealEd",
