World Bookmark - Improved error handling/reporting & other improvements
* Skip capturing/restoring states if dependant states failed to perform their operation
* Added extensive logging and toast notifications
* Prevent bookmarks from being updated or restored during PIE
* Prevent bookmarks from pointing to unsaved worlds
* Bookmarks created from the Content Browser are now automatically populated with the current editor state
* ActorEditorContext state now has a boolean to allow users to prevent that state from being applied when loading the bookmark

diff --git a/Engine/Source/Editor/DataLayerEditor/Private/DataLayer/DataLayerEditorState.cpp b/Engine/Source/Editor/DataLayerEditor/Private/DataLayer/DataLayerEditorState.cpp
--- a/Engine/Source/Editor/DataLayerEditor/Private/DataLayer/DataLayerEditorState.cpp
+++ b/Engine/Source/Editor/DataLayerEditor/Private/DataLayer/DataLayerEditorState.cpp
@@ -20,100 +20,104 @@ FText UDataLayerEditorState::GetCategoryText() const
 	return FText(LOCTEXT("DataLayersEditorStateCategoryText", "Data Layers"));
 }
 
-bool UDataLayerEditorState::CaptureState()
+UEditorState::FOperationResult UDataLayerEditorState::CaptureState()
 {
-	if (UDataLayerManager* DataLayerManager = UDataLayerManager::GetDataLayerManager(GetStateWorld()))
+	UDataLayerManager* DataLayerManager = UDataLayerManager::GetDataLayerManager(GetStateWorld());
+	if (!DataLayerManager)
 	{
-		DataLayerManager->ForEachDataLayerInstance([this](const UDataLayerInstance* DataLayerInstance)
+		return FOperationResult(FOperationResult::Skipped, LOCTEXT("CaptureStateSkipped_NoDataLayerManager", "No data layer manager, world is probably not partitioned"));
+	}
+
+	DataLayerManager->ForEachDataLayerInstance([this](const UDataLayerInstance* DataLayerInstance)
+	{
+		if (DataLayerInstance && DataLayerInstance->GetAsset())
 		{
-			if (DataLayerInstance && DataLayerInstance->GetAsset())
+			if (DataLayerInstance->IsLoadedInEditor() != DataLayerInstance->IsInitiallyLoadedInEditor())
 			{
-				if (DataLayerInstance->IsLoadedInEditor() != DataLayerInstance->IsInitiallyLoadedInEditor())
+				if (DataLayerInstance->IsLoadedInEditor())
 				{
-					if (DataLayerInstance->IsLoadedInEditor())
-					{
-						LoadedDataLayers.Add(DataLayerInstance->GetAsset());
-					}
-					else
-					{
-						NotLoadedDataLayers.Add(DataLayerInstance->GetAsset());
-					}
+					LoadedDataLayers.Add(DataLayerInstance->GetAsset());
+				}
+				else
+				{
+					NotLoadedDataLayers.Add(DataLayerInstance->GetAsset());
 				}
 			}
-
-			return true;
-		});
+		}
 
 		return true;
-	}
+	});
 
-	return false;
+	return FOperationResult(FOperationResult::Success, FText::Format(LOCTEXT("CaptureStateSuccess", "LoadedDataLayers={0}, NotLoadedDataLayers={1}"), LoadedDataLayers.Num(), NotLoadedDataLayers.Num()));
 }
 
-bool UDataLayerEditorState::RestoreState() const
+UEditorState::FOperationResult UDataLayerEditorState::RestoreState() const
 {
-	if (UDataLayerManager* DataLayerManager = UDataLayerManager::GetDataLayerManager(GetStateWorld()))
+	UDataLayerManager* DataLayerManager = UDataLayerManager::GetDataLayerManager(GetStateWorld());
+	if (!DataLayerManager)
 	{
-		// Gather overriden state for all DataLayer instance
-		TMap<const UDataLayerInstance*, bool> DataLayersLoadedInEditor;
+		return FOperationResult(FOperationResult::Skipped, LOCTEXT("RestoreStateSkipped_NoDataLayerManager", "No data layer manager, world is probably not partitioned"));
+	}
 
-		// Unloaded
-		for (const TObjectPtr<const UDataLayerAsset>& DataLayerAsset : NotLoadedDataLayers)
+	// Gather overriden state for all DataLayer instance
+	TMap<const UDataLayerInstance*, bool> DataLayersLoadedInEditor;
+
+	// Unloaded
+	for (const TObjectPtr<const UDataLayerAsset>& DataLayerAsset : NotLoadedDataLayers)
+	{
+		if (DataLayerAsset != nullptr)
 		{
-			if (DataLayerAsset != nullptr)
+			const UDataLayerInstance* DataLayerInstance = DataLayerManager->GetDataLayerInstance(DataLayerAsset);
+			if (DataLayerInstance != nullptr)
 			{
-				const UDataLayerInstance* DataLayerInstance = DataLayerManager->GetDataLayerInstance(DataLayerAsset);
-				if (DataLayerInstance != nullptr)
-				{
-					DataLayersLoadedInEditor.Emplace(DataLayerInstance, false);
-				}
+				DataLayersLoadedInEditor.Emplace(DataLayerInstance, false);
 			}
 		}
+	}
 
-		// Loaded
-		for (const TObjectPtr<const UDataLayerAsset>& DataLayerAsset : LoadedDataLayers)
+	// Loaded
+	for (const TObjectPtr<const UDataLayerAsset>& DataLayerAsset : LoadedDataLayers)
+	{
+		if (DataLayerAsset != nullptr)
 		{
-			if (DataLayerAsset != nullptr)
+			const UDataLayerInstance* DataLayerInstance = DataLayerManager->GetDataLayerInstance(DataLayerAsset);
+			if (DataLayerInstance != nullptr)
 			{
-				const UDataLayerInstance* DataLayerInstance = DataLayerManager->GetDataLayerInstance(DataLayerAsset);
-				if (DataLayerInstance != nullptr)
-				{
-					DataLayersLoadedInEditor.Emplace(DataLayerInstance, true);
-				}
+				DataLayersLoadedInEditor.Emplace(DataLayerInstance, true);
 			}
 		}
+	}
 
-		TArray<UDataLayerInstance*> DataLayersLoaded;
-		TArray<UDataLayerInstance*> DataLayersUnloaded;
-		DataLayerManager->ForEachDataLayerInstance([&DataLayersLoaded, &DataLayersUnloaded, &DataLayersLoadedInEditor](UDataLayerInstance* DataLayerInstance)
+	TArray<UDataLayerInstance*> DataLayersLoaded;
+	TArray<UDataLayerInstance*> DataLayersUnloaded;
+	DataLayerManager->ForEachDataLayerInstance([&DataLayersLoaded, &DataLayersUnloaded, &DataLayersLoadedInEditor](UDataLayerInstance* DataLayerInstance)
+	{
+		bool bDataLayerLoadedInEditor = DataLayerInstance->IsInitiallyLoadedInEditor();
+		if (bool* bBookmarkValue = DataLayersLoadedInEditor.Find(DataLayerInstance))
 		{
-			bool bDataLayerLoadedInEditor = DataLayerInstance->IsInitiallyLoadedInEditor();
-			if (bool* bBookmarkValue = DataLayersLoadedInEditor.Find(DataLayerInstance))
-			{
-				bDataLayerLoadedInEditor = *bBookmarkValue;
-			}
+			bDataLayerLoadedInEditor = *bBookmarkValue;
+		}
 
-			TArray<UDataLayerInstance*>& DataLayersArray = bDataLayerLoadedInEditor ? DataLayersLoaded : DataLayersUnloaded;
-			DataLayersArray.Add(DataLayerInstance);
+		TArray<UDataLayerInstance*>& DataLayersArray = bDataLayerLoadedInEditor ? DataLayersLoaded : DataLayersUnloaded;
+		DataLayersArray.Add(DataLayerInstance);
 
-			// Visibility of Data Layers is not currently saved to users settings or world bookmarks
-			// Until then, we restore them to their initial visibility
-			DataLayerInstance->SetVisible(DataLayerInstance->IsInitiallyVisible());
+		// Visibility of Data Layers is not currently saved to users settings or world bookmarks
+		// Until then, we restore them to their initial visibility
+		DataLayerInstance->SetVisible(DataLayerInstance->IsInitiallyVisible());
 
-			return true;
-		});
+		return true;
+	});
 
-		if (UDataLayerEditorSubsystem* DataLayerEditorSubsystem = UDataLayerEditorSubsystem::Get())
-		{
-			const bool LOAD_DATALAYER = true;
-			const bool UNLOAD_DATALAYER = false;
-			const bool FROM_USER_CHANGE = true;
-			DataLayerEditorSubsystem->SetDataLayersIsLoadedInEditor(DataLayersLoaded, LOAD_DATALAYER, FROM_USER_CHANGE);
-			DataLayerEditorSubsystem->SetDataLayersIsLoadedInEditor(DataLayersUnloaded, UNLOAD_DATALAYER, FROM_USER_CHANGE);
-		}
+	if (UDataLayerEditorSubsystem* DataLayerEditorSubsystem = UDataLayerEditorSubsystem::Get())
+	{
+		const bool LOAD_DATALAYER = true;
+		const bool UNLOAD_DATALAYER = false;
+		const bool FROM_USER_CHANGE = true;
+		DataLayerEditorSubsystem->SetDataLayersIsLoadedInEditor(DataLayersLoaded, LOAD_DATALAYER, FROM_USER_CHANGE);
+		DataLayerEditorSubsystem->SetDataLayersIsLoadedInEditor(DataLayersUnloaded, UNLOAD_DATALAYER, FROM_USER_CHANGE);
 	}
-
-	return true;
+	
+	return FOperationResult(FOperationResult::Success, FText::Format(LOCTEXT("RestoreStateSuccess", "LoadedDataLayers={0}, NotLoadedDataLayers={1}"), LoadedDataLayers.Num(), NotLoadedDataLayers.Num()));
 }
 
 #undef LOCTEXT_NAMESPACE
diff --git a/Engine/Source/Editor/DataLayerEditor/Public/DataLayer/DataLayerEditorState.h b/Engine/Source/Editor/DataLayerEditor/Public/DataLayer/DataLayerEditorState.h
--- a/Engine/Source/Editor/DataLayerEditor/Public/DataLayer/DataLayerEditorState.h
+++ b/Engine/Source/Editor/DataLayerEditor/Public/DataLayer/DataLayerEditorState.h
@@ -15,8 +15,8 @@ class UDataLayerEditorState : public UWorldDependantEditorState
 	//~ Begin UEditorState Interface
 	virtual FText GetCategoryText() const override;
 private:
-	virtual bool CaptureState() override;
-	virtual bool RestoreState() const override;
+	virtual FOperationResult CaptureState() override;
+	virtual FOperationResult RestoreState() const override;
 	//~ End UEditorState Interface
 
 private:
diff --git a/Engine/Source/Editor/LevelEditor/Private/LevelEditorCameraEditorState.cpp b/Engine/Source/Editor/LevelEditor/Private/LevelEditorCameraEditorState.cpp
--- a/Engine/Source/Editor/LevelEditor/Private/LevelEditorCameraEditorState.cpp
+++ b/Engine/Source/Editor/LevelEditor/Private/LevelEditorCameraEditorState.cpp
@@ -19,42 +19,44 @@ FText ULevelEditorCameraEditorState::GetCategoryText() const
 	return FText(LOCTEXT("LevelEditorCameraEditorStateCategoryText", "Camera"));
 }
 
-bool ULevelEditorCameraEditorState::CaptureState()
+UEditorState::FOperationResult ULevelEditorCameraEditorState::CaptureState()
 {
 	TSharedPtr<SLevelViewport> LevelViewport = InternalEditorLevelLibrary::GetActiveLevelViewport();
-	if (LevelViewport.IsValid())
+	if (!LevelViewport.IsValid())
 	{
-		const FLevelEditorViewportClient& LevelViewportClient = LevelViewport->GetLevelViewportClient();
-		CameraLocation = LevelViewportClient.GetViewLocation();
-		CameraRotation = LevelViewportClient.GetViewRotation();
-		CameraFOVAngle = LevelViewportClient.FOVAngle;
-		return true;
-	}
-	else
-	{
-		return false;
+		return FOperationResult(FOperationResult::Failure, LOCTEXT("CaptureStateFailure_NoActiveViewport", "No active viewport"));
 	}
+
+	const FLevelEditorViewportClient& LevelViewportClient = LevelViewport->GetLevelViewportClient();
+	CameraLocation = LevelViewportClient.GetViewLocation();
+	CameraRotation = LevelViewportClient.GetViewRotation();
+	CameraFOVAngle = LevelViewportClient.FOVAngle;
+
+	return FOperationResult(FOperationResult::Success, FText::Format(LOCTEXT("CaptureStateSuccess", "Location=({0}) Rotation=({1}) FOV={2}"), FText::FromString(CameraLocation.ToString()), FText::FromString(CameraRotation.ToString()), CameraFOVAngle));
 }
 
-bool ULevelEditorCameraEditorState::RestoreState() const
+UEditorState::FOperationResult ULevelEditorCameraEditorState::RestoreState() const
 {
-	// Set all level editing cameras to this bookmark
-	for (FLevelEditorViewportClient* LevelVC : GEditor->GetLevelViewportClients())
+	TSharedPtr<SLevelViewport> LevelViewport = InternalEditorLevelLibrary::GetActiveLevelViewport();
+	if (!LevelViewport.IsValid())
 	{
-		LevelVC->SetViewLocation(CameraLocation);
-		if (!LevelVC->IsOrtho())
-		{
-			LevelVC->SetViewRotation(CameraRotation);
-		}
-
-		LevelVC->FOVAngle = CameraFOVAngle;
-		LevelVC->ViewFOV = CameraFOVAngle;
-		LevelVC->Invalidate();
+		return FOperationResult(FOperationResult::Failure, LOCTEXT("RestoreStateFailure_NoActiveViewport", "No active viewport"));
+	}
 
-		FEditorDelegates::OnEditorCameraMoved.Broadcast(CameraLocation, CameraRotation, LevelVC->ViewportType, LevelVC->ViewIndex);
+	FLevelEditorViewportClient& LevelViewportClient = LevelViewport->GetLevelViewportClient();
+	LevelViewportClient.SetViewLocation(CameraLocation);
+	if (!LevelViewportClient.IsOrtho())
+	{
+		LevelViewportClient.SetViewRotation(CameraRotation);
 	}
 
-	return true;
+	LevelViewportClient.FOVAngle = CameraFOVAngle;
+	LevelViewportClient.ViewFOV = CameraFOVAngle;
+	LevelViewportClient.Invalidate();
+
+	FEditorDelegates::OnEditorCameraMoved.Broadcast(CameraLocation, CameraRotation, LevelViewportClient.ViewportType, LevelViewportClient.ViewIndex);
+
+	return FOperationResult(FOperationResult::Success, FText::Format(LOCTEXT("RestoreStateSuccess", "Location=({0}) Rotation=({1}) FOV={2}"), FText::FromString(CameraLocation.ToString()), FText::FromString(CameraRotation.ToString()), CameraFOVAngle));
 }
 
 #undef LOCTEXT_NAMESPACE
diff --git a/Engine/Source/Editor/LevelEditor/Public/LevelEditorCameraEditorState.h b/Engine/Source/Editor/LevelEditor/Public/LevelEditorCameraEditorState.h
--- a/Engine/Source/Editor/LevelEditor/Public/LevelEditorCameraEditorState.h
+++ b/Engine/Source/Editor/LevelEditor/Public/LevelEditorCameraEditorState.h
@@ -13,8 +13,8 @@ class ULevelEditorCameraEditorState : public UEditorState
 	//~ Begin UEditorState Interface
 	virtual FText GetCategoryText() const override;
 private:
-	virtual bool CaptureState() override;
-	virtual bool RestoreState() const override;
+	virtual FOperationResult CaptureState() override;
+	virtual FOperationResult RestoreState() const override;
 	//~ End UEditorState Interface
 
 private:
diff --git a/Engine/Source/Editor/UnrealEd/Private/EditorState/ActorEditorContextEditorState.cpp b/Engine/Source/Editor/UnrealEd/Private/EditorState/ActorEditorContextEditorState.cpp
--- a/Engine/Source/Editor/UnrealEd/Private/EditorState/ActorEditorContextEditorState.cpp
+++ b/Engine/Source/Editor/UnrealEd/Private/EditorState/ActorEditorContextEditorState.cpp
@@ -14,6 +14,7 @@
 
 UActorEditorContextEditorState::UActorEditorContextEditorState(const FObjectInitializer& ObjectInitializer)
 	: Super(ObjectInitializer)
+	, bApplyActorEditorContextOnLoad(true)
 {
 }
 
@@ -39,7 +40,7 @@ FText UActorEditorContextEditorState::GetCategoryText() const
 	return FText(LOCTEXT("ActorEditorContextEditorStateCategoryText", "Actor Editor Context"));
 }
 
-bool UActorEditorContextEditorState::CaptureState()
+UEditorState::FOperationResult UActorEditorContextEditorState::CaptureState()
 {
 	UWorld* CurrentWorld = GetStateWorld();
 
@@ -59,13 +60,18 @@ bool UActorEditorContextEditorState::CaptureState()
 	ActorEditorContext->SetExternalPackage(nullptr);
 	ActorEditorContext->Rename(nullptr, this, REN_DoNotDirty | REN_DontCreateRedirectors | REN_NonTransactional);
 
-	return true;
+	return FOperationResult(FOperationResult::Success);
 }
 
-bool UActorEditorContextEditorState::RestoreState() const
+UEditorState::FOperationResult UActorEditorContextEditorState::RestoreState() const
 {
 	UWorld* CurrentWorld = GetStateWorld();
 
+	if (!bApplyActorEditorContextOnLoad)
+	{
+		return FOperationResult(FOperationResult::Skipped, LOCTEXT("RestoreStateSkipped_ApplyActorEditorContextOnLoad", "User manually disabled application of the actor editor context"));
+	}
+
 	// Actor Editor Context - Apply context from saved actor with EActorEditorContextAction::InitializeContextFromActor
 	ActorEditorContext->Rename(nullptr, CurrentWorld->PersistentLevel, REN_DoNotDirty | REN_DontCreateRedirectors | REN_NonTransactional);
 	ActorEditorContext->SetExternalPackage(GetPackage());
@@ -76,7 +82,7 @@ bool UActorEditorContextEditorState::RestoreState() const
 	ActorEditorContext->SetExternalPackage(nullptr);
 	ActorEditorContext->Rename(nullptr, const_cast<UActorEditorContextEditorState*>(this), REN_DoNotDirty | REN_DontCreateRedirectors | REN_NonTransactional);
 
-	return true;
+	return FOperationResult(FOperationResult::Success);
 }
 
 #undef LOCTEXT_NAMESPACE
diff --git a/Engine/Source/Editor/UnrealEd/Private/EditorState/EditorStateSubsystem.cpp b/Engine/Source/Editor/UnrealEd/Private/EditorState/EditorStateSubsystem.cpp
--- a/Engine/Source/Editor/UnrealEd/Private/EditorState/EditorStateSubsystem.cpp
+++ b/Engine/Source/Editor/UnrealEd/Private/EditorState/EditorStateSubsystem.cpp
@@ -3,9 +3,15 @@
 #include "EditorState/EditorStateSubsystem.h"
 #include "EditorState/WorldEditorState.h"
 #include "Editor.h"
+#include "Framework/Notifications/NotificationManager.h"
+#include "Widgets/Notifications/SNotificationList.h"
 
 #include UE_INLINE_GENERATED_CPP_BY_NAME(EditorStateSubsystem)
 
+DEFINE_LOG_CATEGORY_STATIC(LogEditorState, All, All);
+
+#define LOCTEXT_NAMESPACE "EditorStateSubsystem"
+
 UEditorStateSubsystem* UEditorStateSubsystem::Get()
 {
 	check(GEditor);
@@ -35,32 +41,44 @@ void UEditorStateSubsystem::UnregisterEditorStateType(TSubclassOf<UEditorState>
 	check(RegisteredEditorStateTypes.Contains(InEditorStateType));
 	RegisteredEditorStateTypes.RemoveSwap(InEditorStateType);
 }
-	
-void UEditorStateSubsystem::CaptureEditorState(FEditorStateCollection& OutState, UObject* InStateOuter)
+
+void UEditorStateSubsystem::OutputOperationResult(const UEditorState* InState, bool bRestore, const UEditorState::FOperationResult& InOperationResult) const
 {
-	// Keep track of restored states to avoid cycles
-	TSet<TSubclassOf<UEditorState>> ProcessedStates;
+	const FString ColonSeparator = InOperationResult.GetResultText().IsEmpty() ? TEXT("") : TEXT(": ");
 
-	// Process each state, but capture dependencies first
-	const auto CaptureState = [&ProcessedStates, &OutState](UEditorState* StateToCapture)
+	switch (InOperationResult.GetResult())
 	{
-		auto CaptureStateImpl = [&ProcessedStates, &OutState](UEditorState* StateToCapture, auto& CaptureStateRef) -> void
+	case UEditorState::FOperationResult::Success:
+		UE_LOG(LogEditorState, Log, TEXT("[%s] Success %s%s"), *InState->GetCategoryText().ToString(), *ColonSeparator, *InOperationResult.GetResultText().ToString());
+		break;
+	case UEditorState::FOperationResult::Skipped:
+		UE_LOG(LogEditorState, Log, TEXT("[%s] Skipping %s%s"), *InState->GetCategoryText().ToString(), *ColonSeparator, *InOperationResult.GetResultText().ToString());
+		break;
+	case UEditorState::FOperationResult::Warning:
+		UE_LOG(LogEditorState, Warning, TEXT("[%s] Warning %s%s"), *InState->GetCategoryText().ToString(), *ColonSeparator, *InOperationResult.GetResultText().ToString());
+		break;
+	case UEditorState::FOperationResult::Failure:
 		{
-			if (!ProcessedStates.Contains(StateToCapture->GetClass()))
-			{
-				for (TSubclassOf<UEditorState> DependencyType : StateToCapture->GetDependencies())
-				{
-					CaptureStateRef(OutState.EditorStates.FindChecked(DependencyType), CaptureStateRef);
-				}
-
-				ProcessedStates.Add(StateToCapture->GetClass());
-				StateToCapture->CaptureState();
-			}
-		};
-
-		CaptureStateImpl(StateToCapture, CaptureStateImpl);
-	};
-
+			UE_LOG(LogEditorState, Error, TEXT("[%s] Failure %s%s"), *InState->GetCategoryText().ToString(), *ColonSeparator, *InOperationResult.GetResultText().ToString());
+
+			FText OperationFailureText = bRestore ? LOCTEXT("RestoreFailed", "Failed to restore bookmark state!") : LOCTEXT("CaptureFailed", "Failed to capture bookmark state!");
+			FNotificationInfo Info(OperationFailureText);
+			Info.SubText = InOperationResult.GetResultText();
+			Info.bUseSuccessFailIcons = true;
+			Info.ExpireDuration = 5.0f;
+			Info.bFireAndForget = true;
+			Info.bUseLargeFont = false;
+			Info.Image = FCoreStyle::Get().GetBrush(TEXT("MessageLog.Error"));
+			TSharedPtr<SNotificationItem> NotificationPtr = FSlateNotificationManager::Get().AddNotification(Info);
+			NotificationPtr->SetCompletionState(SNotificationItem::CS_Fail);
+		}
+		break;
+	}
+}
+	
+void UEditorStateSubsystem::CaptureEditorState(FEditorStateCollection& OutState, UObject* InStateOuter)
+{
+	UE_LOG(LogEditorState, Log, TEXT("Capturing editor state..."));
 
 	OutState.EditorStates.Reset();
 
@@ -69,33 +87,64 @@ void UEditorStateSubsystem::CaptureEditorState(FEditorStateCollection& OutState,
 		OutState.EditorStates.Emplace(EditorStateType, NewObject<UEditorState>(InStateOuter, EditorStateType));
 	}
 
-	OutState.ForEachState(CaptureState);
+	OutState.ForEachState([&OutState, this](UEditorState* StateToCapture, bool bCapturedDependencies)
+	{
+		bool bCaptureSuccess = bCapturedDependencies;
+		if (bCaptureSuccess)
+		{
+			UEditorState::FOperationResult Result = StateToCapture->CaptureState();
+			OutputOperationResult(StateToCapture, false, Result);
+
+			bCaptureSuccess = Result == UEditorState::FOperationResult::Success ||
+							  Result == UEditorState::FOperationResult::Warning;
+		}
+		else
+		{
+			OutputOperationResult(StateToCapture, false, UEditorState::FOperationResult(UEditorState::FOperationResult::Skipped, LOCTEXT("CaptureStateSkipped_MissingDependencies", "Missing dependant states, ignoring")));
+		}
+
+		// Couldn't capture state, invalidate it
+		if (!bCaptureSuccess)
+		{
+			StateToCapture->Rename(nullptr, GetTransientPackage());
+			OutState.EditorStates.FindChecked(StateToCapture->GetClass()) = nullptr;
+		}
+
+		return bCaptureSuccess;
+	});
+
+	// Remove invalidated states
+	for (FEditorStateCollection::FEditorStateMap::TIterator It(OutState.EditorStates); It; ++It)
+	{
+		if (It.Value() == nullptr)
+		{
+			It.RemoveCurrent();
+		}
+	}	
+
+	UE_LOG(LogEditorState, Log, TEXT("Captured editor state... DONE"));
 }
 
 void UEditorStateSubsystem::RestoreEditorState(const FEditorStateCollection& InState)
 {
-	// Keep track of restored states to avoid cycles
-	TSet<TSubclassOf<UEditorState>> ProcessedStates;	
-	
-	// Process each state, but restore dependencies first
-	const auto RestoreState = [&ProcessedStates, &InState](const UEditorState* StateToRestore)
+	UE_LOG(LogEditorState, Log, TEXT("Restoring editor state..."));
+
+	InState.ForEachState([this](const UEditorState* StateToRestore, bool bRestoredDependencies)
 	{
-		auto RestoreStateImpl = [&ProcessedStates, &InState](const UEditorState* StateToRestore, auto& RestoreStateRef) -> void
+		if (!bRestoredDependencies)
 		{
-			if (!ProcessedStates.Contains(StateToRestore->GetClass()))
-			{
-				for (TSubclassOf<UEditorState> DependencyType : StateToRestore->GetDependencies())
-				{
-					RestoreStateRef(InState.EditorStates.FindChecked(DependencyType), RestoreStateRef);
-				}
-
-				ProcessedStates.Add(StateToRestore->GetClass());
-				StateToRestore->RestoreState();
-			}
-		};
-
-		RestoreStateImpl(StateToRestore, RestoreStateImpl);
-	};
-
-	InState.ForEachState(RestoreState);
+			OutputOperationResult(StateToRestore, true, UEditorState::FOperationResult(UEditorState::FOperationResult::Skipped, LOCTEXT("RestoreStateSkipped_MissingDependencies", "Missing dependant states, ignoring")));
+			return false;
+		}
+		
+		UEditorState::FOperationResult Result = StateToRestore->RestoreState();
+		OutputOperationResult(StateToRestore, true, Result);
+
+		return Result == UEditorState::FOperationResult::Success ||
+			   Result == UEditorState::FOperationResult::Warning;
+	});
+
+	UE_LOG(LogEditorState, Log, TEXT("Restored editor state... DONE"));
 }
+
+#undef LOCTEXT_NAMESPACE
\ No newline at end of file
diff --git a/Engine/Source/Editor/UnrealEd/Private/EditorState/WorldEditorState.cpp b/Engine/Source/Editor/UnrealEd/Private/EditorState/WorldEditorState.cpp
--- a/Engine/Source/Editor/UnrealEd/Private/EditorState/WorldEditorState.cpp
+++ b/Engine/Source/Editor/UnrealEd/Private/EditorState/WorldEditorState.cpp
@@ -1,6 +1,7 @@
 // Copyright Epic Games, Inc. All Rights Reserved.
 
 #include "EditorState/WorldEditorState.h"
+#include "AssetRegistry/IAssetRegistry.h"
 #include "Engine/World.h"
 #include "FileHelpers.h"
 #include "Misc/MessageDialog.h"
@@ -24,27 +25,44 @@ TSoftObjectPtr<UWorld> UWorldEditorState::GetStateWorld() const
 	return World;
 }
 
-bool UWorldEditorState::CaptureState()
+UEditorState::FOperationResult UWorldEditorState::CaptureState()
 {
+	if (!GWorld || FPackageName::IsTempPackage(GWorld->GetPackage()->GetName()))
+	{
+		return FOperationResult(FOperationResult::Failure, LOCTEXT("CaptureStateFailure_UnsavedWorld", "Current world must be saved"));
+	}
+
 	// Set the current world.
 	World = (UWorld*)GWorld;
-	return true;
+	return FOperationResult(FOperationResult::Success);
 }
 
-bool UWorldEditorState::RestoreState() const
+UEditorState::FOperationResult UWorldEditorState::RestoreState() const
 {
 	if (World.IsNull())
 	{
-		return false;
+		return FOperationResult(FOperationResult::Failure, LOCTEXT("RestoreStateFailure_UnsavedWorld", "Bookmark world is invalid"));
 	}
 
-	UWorld* CurrentWorld = GWorld;
-	UWorld* TargetWorld = Cast<UWorld>(World.ToSoftObjectPath().ResolveObject());
+	FString WorldPackageString = World.GetLongPackageName();
+	FText WorldPackageText = FText::FromString(WorldPackageString);
+	FAssetData WorldAsset;
+
+	if (IAssetRegistry::GetChecked().TryGetAssetByObjectPath(World.ToSoftObjectPath(), WorldAsset) != UE::AssetRegistry::EExists::Exists)
+	{
+		return FOperationResult(FOperationResult::Failure, FText::Format(LOCTEXT("RestoreStateFailure_TargetWorldNotFound", "Bookmark world {0} couldn't be resolved"), WorldPackageText));
+	}
+
+	if (WorldAsset.AssetClassPath != UWorld::StaticClass()->GetClassPathName())
+	{
+		return FOperationResult(FOperationResult::Failure, FText::Format(LOCTEXT("RestoreStateFailure_TargetNotAWorld", "{0} is not a world."), WorldPackageText));
+	}
 
 	// Already in the right world.
-	if (CurrentWorld == TargetWorld)
+	UWorld* CurrentWorld = GWorld;
+	if (CurrentWorld == World)
 	{
-		return true;
+		return FOperationResult(FOperationResult::Success, FText::Format(LOCTEXT("RestoreStateSuccess_WorldAlreadyLoaded", "Bookmark world {0} was already active"), WorldPackageText));
 	}
 
 	// If there are any unsaved changes to the current level, see if the user wants to save those first
@@ -54,30 +72,36 @@ bool UWorldEditorState::RestoreState() const
 	constexpr bool bSaveContentPackages = true;
 	if (FEditorFileUtils::SaveDirtyPackages(bPromptUserToSave, bSaveMapPackages, bSaveContentPackages))
 	{
+		UWorld* TargetWorld = World.Get();
 		if (TargetWorld != nullptr)
 		{
 			// Validate that Asset was saved or isn't loaded meaning it can be loaded
 			if (TargetWorld->GetPackage() == nullptr)
 			{
-				FMessageDialog::Open(EAppMsgType::Ok, FText::Format(LOCTEXT("CannotOpenNotInPackage", "World ({0}) is not in a package and cannot be opened."), FText::FromString(TargetWorld->GetFullName())));
-				return false;
+				return FOperationResult(FOperationResult::Failure, FText::Format(LOCTEXT("RestoreStateFailure_CannotOpenNotInPackage", "The world you are trying to open ({0}) needs to be saved first."), WorldPackageText));
 			}
 
 			if (TargetWorld->GetPackage()->HasAnyPackageFlags(PKG_NewlyCreated))
 			{
-				FMessageDialog::Open(EAppMsgType::Ok, LOCTEXT("CannotOpenNewlyCreatedMapWithoutSaving", "The world you are trying to open needs to be saved first."));
-				return false;
+				return FOperationResult(FOperationResult::Failure, FText::Format(LOCTEXT("RestoreStateFailure_CannotOpenNewlyCreatedMapWithoutSaving", "The world you are trying to open ({0}) needs to be saved first."), WorldPackageText));
 			}
 		}
 	}
 	else
 	{
-		return false;
+		return FOperationResult(FOperationResult::Skipped, LOCTEXT("RestoreStateSkipped_UserDeclinedToSave", "Declined to save dirty packages"));
 	}
 
 	const bool bLoadAsTemplate = false;
 	const bool bShowProgress = true;
-	return FEditorFileUtils::LoadMap(FPackageName::LongPackageNameToFilename(World.GetLongPackageName(), FPackageName::GetMapPackageExtension()), bLoadAsTemplate, bShowProgress);
+	const bool bMapLoaded = FEditorFileUtils::LoadMap(FPackageName::LongPackageNameToFilename(WorldPackageString, FPackageName::GetMapPackageExtension()), bLoadAsTemplate, bShowProgress);
+
+	if (!bMapLoaded)
+	{
+		return FOperationResult(FOperationResult::Failure, FText::Format(LOCTEXT("RestoreStateFailure_LoadWorld", "Failed to load world {0}"), WorldPackageText));
+	}
+	
+	return FOperationResult(FOperationResult::Success, FText::Format(LOCTEXT("RestoreStateSuccess_WorldLoaded", "Loaded world {0}"), WorldPackageText));
 }
 	
 UWorldDependantEditorState::UWorldDependantEditorState(const FObjectInitializer& ObjectInitializer)
diff --git a/Engine/Source/Editor/UnrealEd/Public/EditorState/ActorEditorContextEditorState.h b/Engine/Source/Editor/UnrealEd/Public/EditorState/ActorEditorContextEditorState.h
--- a/Engine/Source/Editor/UnrealEd/Public/EditorState/ActorEditorContextEditorState.h
+++ b/Engine/Source/Editor/UnrealEd/Public/EditorState/ActorEditorContextEditorState.h
@@ -17,8 +17,8 @@ public:
 	virtual FText GetCategoryText() const override;
 
 protected:
-	virtual bool CaptureState() override;
-	virtual bool RestoreState() const override;
+	virtual FOperationResult CaptureState() override;
+	virtual FOperationResult RestoreState() const override;
 	//~ End UEditorState Interface
 
 	//~ Begin UObject Interface
@@ -29,4 +29,7 @@ protected:
 	// Actor used to store the editor context, will be used to restore context when loading the bookmark.
 	UPROPERTY(DuplicateTransient)
 	TObjectPtr<AActor> ActorEditorContext;
+
+	UPROPERTY(EditAnywhere, Category = ActorEditorContext)
+	bool bApplyActorEditorContextOnLoad;
 };
diff --git a/Engine/Source/Editor/UnrealEd/Public/EditorState/EditorState.h b/Engine/Source/Editor/UnrealEd/Public/EditorState/EditorState.h
--- a/Engine/Source/Editor/UnrealEd/Public/EditorState/EditorState.h
+++ b/Engine/Source/Editor/UnrealEd/Public/EditorState/EditorState.h
@@ -16,19 +16,61 @@ class UEditorState : public UObject
 {
 	GENERATED_BODY()
 
+public:
+	struct FOperationResult
+	{
+		enum EResult
+		{
+			Success,
+			Skipped,
+			Warning,
+			Failure
+		};
+
+		explicit FOperationResult(EResult InResult, const FText& InResultText = FText())
+			: Result(InResult)
+			, ResultText(InResultText)
+		{
+		}
+
+		bool operator == (EResult InResult) const
+		{
+			return Result == InResult;
+		}
+
+		bool operator != (EResult InResult) const
+		{
+			return Result == InResult;
+		}
+
+		EResult GetResult() const
+		{
+			return Result;
+		}
+
+		const FText& GetResultText() const
+		{
+			return ResultText;
+		}
+
+	private:
+		EResult	 Result;
+		FText	 ResultText;
+	};
+
 public:
 	/** Get the category under which this state's properties should be displayed. */
 	UNREALED_API virtual FText GetCategoryText() const PURE_VIRTUAL(UEditorState::GetCategoryText, return FText::GetEmpty(); );
 	
-protected:
 	/** Get a list of state types this editor state depends on. Dependent types will be restored after their dependencies. */
 	virtual TArray<TSubclassOf<UEditorState>> GetDependencies() const { return {}; }
 
+protected:
 	/** Capture the state of the editor. Must be subclassed. */
-	virtual bool CaptureState() PURE_VIRTUAL(UEditorState::CaptureState, return false; );
+	virtual FOperationResult CaptureState() PURE_VIRTUAL(UEditorState::CaptureState, return FOperationResult(FOperationResult::Failure); );
 
 	/** Restore the state of the editor. Must be subclassed. */
-	virtual bool RestoreState() const PURE_VIRTUAL(UEditorState::RestoreState, return false; );
+	virtual FOperationResult RestoreState() const PURE_VIRTUAL(UEditorState::RestoreState, return FOperationResult(FOperationResult::Failure); );
 
 private:
 	/** Prevent access to UObject::GetWorld() as it's not relevant for this class and is error prone. */
diff --git a/Engine/Source/Editor/UnrealEd/Public/EditorState/EditorStateCollection.h b/Engine/Source/Editor/UnrealEd/Public/EditorState/EditorStateCollection.h
--- a/Engine/Source/Editor/UnrealEd/Public/EditorState/EditorStateCollection.h
+++ b/Engine/Source/Editor/UnrealEd/Public/EditorState/EditorStateCollection.h
@@ -34,14 +34,18 @@ struct FEditorStateCollection
 		return *CastChecked<TEditorStateType>(EditorState.Get());
 	}
 
-	void ForEachState(TFunctionRef<void(const UEditorState*)> InFunc) const
+	/** Iterate over each state, making sure iteration is done over dependant states first */
+	void ForEachState(TFunctionRef<bool(const UEditorState*, bool)> InFunc) const
 	{
-		for (const TPair<TSubclassOf<UEditorState>, TObjectPtr<UEditorState>>& KeyValue : EditorStates)
-		{
-			InFunc(KeyValue.Value);
-		}
+		ForEachState(this, InFunc);
 	}
 
+	/** Iterate over each state, making sure iteration is done over dependant states first */
+	void ForEachState(TFunctionRef<bool(UEditorState*, bool)> InFunc)
+	{
+		ForEachState(this, InFunc);
+	}
+	
 	TArray<UEditorState*> GetStates() const
 	{
 		TArray<UEditorState*> Values;
@@ -55,14 +59,56 @@ struct FEditorStateCollection
 	}
 
 private:
-	void ForEachState(TFunctionRef<void(UEditorState*)> InFunc)
+	template <typename ThisType, typename FuncType>
+	static void ForEachState(ThisType* Self, FuncType InFunc)
 	{
-		for (const TPair<TSubclassOf<UEditorState>, TObjectPtr<UEditorState>>& KeyValue : EditorStates)
+		// Keep track of processed states to avoid cycles
+		TMap<TSubclassOf<UEditorState>, bool> ProcessedStates;
+
+		const auto ProcessState = [&ProcessedStates, &InFunc, Self](UEditorState* StateToProcess)
 		{
-			InFunc(KeyValue.Value);
+			auto ProcessStateImpl = [&ProcessedStates, &InFunc, Self](UEditorState* StateToProcess, auto& ProcessStateRef) -> bool
+			{
+				// Avoid processing dependencies twice
+				bool* bProcessedStateResult = ProcessedStates.Find(StateToProcess->GetClass());
+				if (!bProcessedStateResult)
+				{
+					bool bProcessedDependenciesSuccessfully = true;
+
+					// Process the dependencies - Stop if any error is encountered
+					for (TSubclassOf<UEditorState> DependencyType : StateToProcess->GetDependencies())
+					{
+						UEditorState* DependencyState = Self->EditorStates.FindRef(DependencyType);
+						bProcessedDependenciesSuccessfully = DependencyState && ProcessStateRef(DependencyState, ProcessStateRef);
+
+						if (!bProcessedDependenciesSuccessfully)
+						{
+							break;
+						}
+					}
+
+					// Process the state itself
+					bool bSuccess = InFunc(StateToProcess, bProcessedDependenciesSuccessfully);
+					ProcessedStates.Add(StateToProcess->GetClass(), bSuccess);
+					return bSuccess;
+				}
+				else
+				{
+					return *bProcessedStateResult;
+				}
+			};
+
+			return ProcessStateImpl(StateToProcess, ProcessStateImpl);
+		};
+
+		for (const auto& KeyValue : Self->EditorStates)
+		{
+			ProcessState(KeyValue.Value);
 		}
 	}
 
+	typedef TMap<TSubclassOf<UEditorState>, TObjectPtr<UEditorState>> FEditorStateMap;
+
 	UPROPERTY()
 	TMap<TSubclassOf<UEditorState>, TObjectPtr<UEditorState>> EditorStates;
 	
diff --git a/Engine/Source/Editor/UnrealEd/Public/EditorState/EditorStateSubsystem.h b/Engine/Source/Editor/UnrealEd/Public/EditorState/EditorStateSubsystem.h
--- a/Engine/Source/Editor/UnrealEd/Public/EditorState/EditorStateSubsystem.h
+++ b/Engine/Source/Editor/UnrealEd/Public/EditorState/EditorStateSubsystem.h
@@ -54,5 +54,7 @@ private:
 	UNREALED_API void RegisterEditorStateType(TSubclassOf<UEditorState> InEditorStateType);
 	UNREALED_API void UnregisterEditorStateType(TSubclassOf<UEditorState> InEditorStateType);
 
+	void OutputOperationResult(const UEditorState* InState, bool bRestore, const UEditorState::FOperationResult& InOperationResult) const;
+
 	TArray<TSubclassOf<UEditorState>> RegisteredEditorStateTypes;
 };
diff --git a/Engine/Source/Editor/UnrealEd/Public/EditorState/WorldEditorState.h b/Engine/Source/Editor/UnrealEd/Public/EditorState/WorldEditorState.h
--- a/Engine/Source/Editor/UnrealEd/Public/EditorState/WorldEditorState.h
+++ b/Engine/Source/Editor/UnrealEd/Public/EditorState/WorldEditorState.h
@@ -17,8 +17,8 @@ class UWorldEditorState : public UEditorState
 	//~ Begin UEditorState Interface
 	virtual FText GetCategoryText() const override;
 private:
-	virtual bool CaptureState() override;
-	virtual bool RestoreState() const override;
+	virtual FOperationResult CaptureState() override;
+	virtual FOperationResult RestoreState() const override;
 	//~ End UEditorState Interface
 
 	// World associated with the bookmark
diff --git a/Engine/Source/Editor/WorldBookmark/Private/WorldBookmark/AssetDefinition_WorldBookmark.cpp b/Engine/Source/Editor/WorldBookmark/Private/WorldBookmark/AssetDefinition_WorldBookmark.cpp
--- a/Engine/Source/Editor/WorldBookmark/Private/WorldBookmark/AssetDefinition_WorldBookmark.cpp
+++ b/Engine/Source/Editor/WorldBookmark/Private/WorldBookmark/AssetDefinition_WorldBookmark.cpp
@@ -45,37 +45,43 @@ EAssetCommandResult UAssetDefinition_WorldBookmark::OpenAssets(const FAssetOpenA
 //--------------------------------------------------------------------
 namespace MenuExtension_Bookmark
 {	
-	bool CanExecuteBookmarkAction(const FToolMenuContext& InContext)
+	UWorldBookmark* GetBookmarkFromContext(const FToolMenuContext& InContext)
 	{
 		if (UContentBrowserAssetContextMenuContext::GetNumAssetsSelected(InContext) != 1)
 		{
-			return false;
+			return nullptr;
 		}
 
 		const UContentBrowserAssetContextMenuContext* Context = UContentBrowserAssetContextMenuContext::FindContextWithAssets(InContext);
 		const TArray<UWorldBookmark*> WorldBookmarks = Context->LoadSelectedObjects<UWorldBookmark>();
-		return WorldBookmarks.Num() == 1 && WorldBookmarks[0] != nullptr;
+		return WorldBookmarks.Num() == 1 ? WorldBookmarks[0] : nullptr;
+	}
+
+	bool CanExecuteGoToWorldBookmark(const FToolMenuContext& InContext)
+	{
+		UWorldBookmark* WorldBookmark = GetBookmarkFromContext(InContext);
+		return WorldBookmark && WorldBookmark->CanLoad();
+	}
+
+	bool CanExecuteUpdateWorldBookmark(const FToolMenuContext& InContext)
+	{
+		UWorldBookmark* WorldBookmark = GetBookmarkFromContext(InContext);
+		return WorldBookmark && WorldBookmark->CanUpdate();
 	}
 
 	void ExecuteGoToWorldBookmark(const FToolMenuContext& InContext)
 	{
-		const UContentBrowserAssetContextMenuContext* Context = UContentBrowserAssetContextMenuContext::FindContextWithAssets(InContext);
-		const TArray<UWorldBookmark*> WorldBookmarks = Context->LoadSelectedObjects<UWorldBookmark>();
-		
-		if (ensure(WorldBookmarks.Num() == 1 && WorldBookmarks[0] != nullptr))
+		if (UWorldBookmark* WorldBookmark = GetBookmarkFromContext(InContext))
 		{
-			WorldBookmarks[0]->Load();
+			WorldBookmark->Load();
 		}
 	}
 
 	void ExecuteUpdateWorldBookmark(const FToolMenuContext& InContext)
 	{
-		const UContentBrowserAssetContextMenuContext* Context = UContentBrowserAssetContextMenuContext::FindContextWithAssets(InContext);
-		const TArray<UWorldBookmark*> WorldBookmarks = Context->LoadSelectedObjects<UWorldBookmark>();
-
-		if (ensure(WorldBookmarks.Num() == 1 && WorldBookmarks[0] != nullptr))
-		{			
-			WorldBookmarks[0]->Update();
+		if (UWorldBookmark* WorldBookmark = GetBookmarkFromContext(InContext))
+		{
+			WorldBookmark->Update();
 		}
 	}
 
@@ -95,7 +101,7 @@ namespace MenuExtension_Bookmark
 
 				FToolUIAction UIAction;
 				UIAction.ExecuteAction = FToolMenuExecuteAction::CreateStatic(&ExecuteGoToWorldBookmark);
-				UIAction.CanExecuteAction = FToolMenuCanExecuteAction::CreateStatic(&CanExecuteBookmarkAction);
+				UIAction.CanExecuteAction = FToolMenuCanExecuteAction::CreateStatic(&CanExecuteGoToWorldBookmark);
 
 				InSection.AddMenuEntry("Bookmark_Edit", Label, ToolTip, Icon, UIAction);
 			}));
@@ -108,7 +114,7 @@ namespace MenuExtension_Bookmark
 					
 				FToolUIAction UIAction;
 				UIAction.ExecuteAction = FToolMenuExecuteAction::CreateStatic(&ExecuteUpdateWorldBookmark);
-				UIAction.CanExecuteAction = FToolMenuCanExecuteAction::CreateStatic(&CanExecuteBookmarkAction);
+				UIAction.CanExecuteAction = FToolMenuCanExecuteAction::CreateStatic(&CanExecuteUpdateWorldBookmark);
 
 				InSection.AddMenuEntry("WorldBookmark_Update", Label, ToolTip, Icon, UIAction);
 			}));
diff --git a/Engine/Source/Editor/WorldBookmark/Private/WorldBookmark/Browser/SWorldBookmarkBrowser.cpp b/Engine/Source/Editor/WorldBookmark/Private/WorldBookmark/Browser/SWorldBookmarkBrowser.cpp
--- a/Engine/Source/Editor/WorldBookmark/Private/WorldBookmark/Browser/SWorldBookmarkBrowser.cpp
+++ b/Engine/Source/Editor/WorldBookmark/Private/WorldBookmark/Browser/SWorldBookmarkBrowser.cpp
@@ -474,9 +474,6 @@ FReply SWorldBookmarkBrowser::OnCreateNewWorldBookmarkClicked()
 	UObject* NewAsset = AssetToolsModule.Get().CreateAssetWithDialog(UWorldBookmark::StaticClass(), Factory);
 	if (UWorldBookmark* NewWorldBookmark = Cast<UWorldBookmark>(NewAsset))
 	{
-		// Update it with the current state of the editor
-		NewWorldBookmark->Update();
-
 		// Show it in the content browser
 		GEditor->SyncBrowserToObject(NewWorldBookmark);
 
diff --git a/Engine/Source/Editor/WorldBookmark/Private/WorldBookmark/WorldBookmark.cpp b/Engine/Source/Editor/WorldBookmark/Private/WorldBookmark/WorldBookmark.cpp
--- a/Engine/Source/Editor/WorldBookmark/Private/WorldBookmark/WorldBookmark.cpp
+++ b/Engine/Source/Editor/WorldBookmark/Private/WorldBookmark/WorldBookmark.cpp
@@ -5,11 +5,14 @@
 #include "EditorState/WorldEditorState.h"
 
 #include "AssetRegistry/AssetRegistryHelpers.h"
-#include "UObject/AssetRegistryTagsContext.h"
+#include "Editor.h"
+#include "EngineUtils.h"
 #include "FileHelpers.h"
+#include "Framework/Notifications/NotificationManager.h"
 #include "Misc/MessageDialog.h"
-#include "EngineUtils.h"
 #include "ScopedTransaction.h"
+#include "UObject/AssetRegistryTagsContext.h"
+#include "Widgets/Notifications/SNotificationList.h"
 
 #include UE_INLINE_GENERATED_CPP_BY_NAME(WorldBookmark)
 
@@ -28,6 +31,11 @@ void UWorldBookmark::PostInitProperties()
 	{
 		BookmarkGuid = FGuid::NewGuid();
 		BookmarkAssetPath = GetPathName();
+
+		if (CanUpdate())
+		{
+			Update();
+		}
 	}
 }
 
@@ -76,34 +84,102 @@ void UWorldBookmark::GetAssetRegistryTags(FAssetRegistryTagsContext Context) con
 	}
 }
 
+bool UWorldBookmark::CanLoad(FText* FailureReason) const
+{
+	if (GEditor->PlayWorld || GIsPlayInEditorWorld)
+	{
+		if (FailureReason)
+		{
+			*FailureReason = LOCTEXT("CanLoadFailure_IsInPIE", "Can't load bookmark while in PIE");
+		}
+
+		return false;
+	}
+
+	return true;
+}
+
 void UWorldBookmark::Load()
 {
-	UEditorStateSubsystem::Get()->RestoreEditorState(EditorState);
+	FText CanLoadFailureReason;
+	if (CanLoad(&CanLoadFailureReason))
+	{
+		UEditorStateSubsystem::Get()->RestoreEditorState(EditorState);
 
-	SetUserLastLoadedTimeStampUTC(FDateTime::UtcNow());
-	SaveConfig();
+		SetUserLastLoadedTimeStampUTC(FDateTime::UtcNow());
+		SaveConfig();
+	}
+	else
+	{
+		FNotificationInfo Info(LOCTEXT("CantLoadBookmark", "Failed to load bookmark!"));
+		Info.SubText = CanLoadFailureReason;
+		Info.ExpireDuration = 5.0f;
+		Info.bFireAndForget = true;
+		Info.bUseLargeFont = false;
+		Info.Image = FCoreStyle::Get().GetBrush(TEXT("MessageLog.Warning"));
+		FSlateNotificationManager::Get().AddNotification(Info);
+	}
 }
 
-void UWorldBookmark::Update()
+bool UWorldBookmark::CanUpdate(FText* FailureReason) const
 {
 	UWorld* CurrentWorld = GWorld;
-	if (CurrentWorld != nullptr && !CurrentWorld->IsGameWorld())
+	if (CurrentWorld == nullptr)
 	{
-		FScopedTransaction ScopedTransaction(LOCTEXT("UpdateWorldBookmarkTransaction", "Update World Bookmark"));
+		if (FailureReason)
+		{
+			*FailureReason = LOCTEXT("CanUpdateFailure_InvalidWorld", "Current world is invalid");
+		}
 
-		UWorld* BookmarkWorld = EditorState.HasState<UWorldEditorState>() ? EditorState.GetState<UWorldEditorState>()->GetStateWorld().Get() : nullptr;
-		if (BookmarkWorld != nullptr && BookmarkWorld != CurrentWorld)
+		return false;
+	}
+	
+	if (GEditor->PlayWorld || GIsPlayInEditorWorld)
+	{
+		if (FailureReason)
+		{
+			*FailureReason = LOCTEXT("CanUpdateFailure_IsInPIE", "Can't update bookmark while in PIE");
+		}
+
+		return false;
+	}
+
+	if (FPackageName::IsTempPackage(CurrentWorld->GetPackage()->GetName()))
+	{
+		if (FailureReason)
 		{
-			FMessageDialog::Open(EAppMsgType::Ok, LOCTEXT("CannotUpdateWorldBookmark", "The bookmark you are trying to update does not belong to the current world."));
-			return;
+			*FailureReason = LOCTEXT("CanUpdateFailure_InvalidWorld", "Unsaved world");
 		}
 
+		return false;
+	}
+
+	return true;
+}
+
+void UWorldBookmark::Update()
+{
+	FText CanUpdateFailureReason;
+	if (CanUpdate(&CanUpdateFailureReason))
+	{
+		FScopedTransaction ScopedTransaction(LOCTEXT("UpdateWorldBookmarkTransaction", "Update World Bookmark"));
+
 		Modify();
 
 		UEditorStateSubsystem::Get()->CaptureEditorState(EditorState, this);
 
 		PostEditChange();
 	}
+	else
+	{
+		FNotificationInfo Info(LOCTEXT("CantUpdateBookmark", "Failed to update bookmark!"));
+		Info.SubText = CanUpdateFailureReason;
+		Info.ExpireDuration = 5.0f;
+		Info.bFireAndForget = true;
+		Info.bUseLargeFont = false;
+		Info.Image = FCoreStyle::Get().GetBrush(TEXT("MessageLog.Warning"));
+		FSlateNotificationManager::Get().AddNotification(Info);
+	}
 }
 
 bool UWorldBookmark::GetIsUserFavorite() const
diff --git a/Engine/Source/Editor/WorldBookmark/Public/WorldBookmark/WorldBookmark.h b/Engine/Source/Editor/WorldBookmark/Public/WorldBookmark/WorldBookmark.h
--- a/Engine/Source/Editor/WorldBookmark/Public/WorldBookmark/WorldBookmark.h
+++ b/Engine/Source/Editor/WorldBookmark/Public/WorldBookmark/WorldBookmark.h
@@ -15,6 +15,12 @@ class UWorldBookmark : public UObject
 {
 	GENERATED_UCLASS_BODY()
 
+	/** Test whether it's possible to load this bookmark given the current state of the editor. */
+	bool CanLoad(FText* FailureReason = nullptr) const;
+
+	/** Test whether it's possible to update this bookmark given the current state of the editor. */
+	bool CanUpdate(FText* FailureReason = nullptr) const;
+
 	/** Load the bookmark data, restoring the editor state (loaded world, camera location, etc) to what is defined by the bookmark. */
 	void Load();
 
diff --git a/Engine/Source/Editor/WorldPartitionEditor/Private/WorldPartition/WorldPartitionEditorState.cpp b/Engine/Source/Editor/WorldPartitionEditor/Private/WorldPartition/WorldPartitionEditorState.cpp
--- a/Engine/Source/Editor/WorldPartitionEditor/Private/WorldPartition/WorldPartitionEditorState.cpp
+++ b/Engine/Source/Editor/WorldPartitionEditor/Private/WorldPartition/WorldPartitionEditorState.cpp
@@ -21,106 +21,123 @@ FText UWorldPartitionEditorState::GetCategoryText() const
 	return FText(LOCTEXT("WorldPartitionEditorStateCategoryText", "World Partition"));
 }
 
-bool UWorldPartitionEditorState::CaptureState()
+UEditorState::FOperationResult UWorldPartitionEditorState::CaptureState()
 {
 	UWorld* CurrentWorld = GetStateWorld();
-	if (CurrentWorld && CurrentWorld->IsPartitionedWorld())
+	if (!ensure(CurrentWorld) || !CurrentWorld->IsPartitionedWorld())
 	{
-		if (GetDefault<UWorldPartitionEditorSettings>()->GetEnableLoadingInEditor())
+		return FOperationResult(FOperationResult::Skipped, LOCTEXT("CaptureStateSkipped_WorldIsNotPartitioned", "World is not partitioned"));
+	}
+
+	if (!GetDefault<UWorldPartitionEditorSettings>()->GetEnableLoadingInEditor())
+	{
+		return FOperationResult(FOperationResult::Skipped, LOCTEXT("CaptureStateSkipped_LoadingInEditorDisabled", "Loading in editor is disabled"));
+	}
+
+	auto GetRegionHash = [](const FBox& LoadedRegion)
+	{
+		return HashCombineFast(GetTypeHash(LoadedRegion.Min), GetTypeHash(LoadedRegion.Max));
+	};
+
+	// Loaded regions
+	int32 NbRegionsLoaded = 0;
+	TSet<uint32> LoadedEditorRegionsHashes;
+	for (const FBox& LoadedRegion : CurrentWorld->GetWorldPartition()->GetUserLoadedEditorRegions())
+	{
+		bool bAlreadyInSet = false;
+		LoadedEditorRegionsHashes.FindOrAdd(GetRegionHash(LoadedRegion), &bAlreadyInSet);
+		if (!bAlreadyInSet)
 		{
-			auto GetRegionHash = [](const FBox& LoadedRegion)
-			{
-				return HashCombineFast(GetTypeHash(LoadedRegion.Min), GetTypeHash(LoadedRegion.Max));
-			};
+			LoadedEditorRegions.Add(LoadedRegion);
+		}
+	}
 
-			// Loaded regions
-			TSet<uint32> LoadedEditorRegionsHashes;
-			for (const FBox& LoadedRegion : CurrentWorld->GetWorldPartition()->GetUserLoadedEditorRegions())
-			{
-				bool bAlreadyInSet = false;
-				LoadedEditorRegionsHashes.FindOrAdd(GetRegionHash(LoadedRegion), &bAlreadyInSet);
-				if (!bAlreadyInSet)
-				{
-					LoadedEditorRegions.Add(LoadedRegion);
-				}
-			}
+	// Loaded location volumes
+	int32 NbVolumesLoaded = 0;
+	for (FActorDescContainerInstanceCollection::TConstIterator<> Iterator(CurrentWorld->GetWorldPartition()); Iterator; ++Iterator)
+	{
+		if (ALocationVolume* LocationVolume = Cast<ALocationVolume>(Iterator->GetActor()); IsValid(LocationVolume))
+		{
+			check(LocationVolume->GetClass()->ImplementsInterface(UWorldPartitionActorLoaderInterface::StaticClass()));
 
-			// Loaded location volumes
-			for (FActorDescContainerInstanceCollection::TConstIterator<> Iterator(CurrentWorld->GetWorldPartition()); Iterator; ++Iterator)
+			IWorldPartitionActorLoaderInterface::ILoaderAdapter* LoaderAdapter = Cast<IWorldPartitionActorLoaderInterface>(LocationVolume)->GetLoaderAdapter();
+			check(LoaderAdapter);
+
+			if (LoaderAdapter->IsLoaded() && LoaderAdapter->GetUserCreated())
 			{
-				if (ALocationVolume* LocationVolume = Cast<ALocationVolume>(Iterator->GetActor()); IsValid(LocationVolume))
-				{
-					check(LocationVolume->GetClass()->ImplementsInterface(UWorldPartitionActorLoaderInterface::StaticClass()));
-
-					IWorldPartitionActorLoaderInterface::ILoaderAdapter* LoaderAdapter = Cast<IWorldPartitionActorLoaderInterface>(LocationVolume)->GetLoaderAdapter();
-					check(LoaderAdapter);
-
-					if (LoaderAdapter->IsLoaded() && LoaderAdapter->GetUserCreated())
-					{
-						LoadedEditorLocationVolumes.Add(LocationVolume->GetFName());
-					}
-				}
+				LoadedEditorLocationVolumes.Add(LocationVolume->GetFName());
 			}
 		}
+	}
 
-		return true;
+	if (NbRegionsLoaded == 0 && NbVolumesLoaded == 0)
+	{
+		return FOperationResult(FOperationResult::Skipped, LOCTEXT("CaptureStateSkipped_NoLoadedRegionsOrVolumes", "No regions or location volumes currently loaded"));
 	}
 
-	return false;
+	return FOperationResult(FOperationResult::Success, FText::Format(LOCTEXT("CaptureStateSuccess", "{0} regions and {1} location volumes currently loaded"), NbRegionsLoaded, NbVolumesLoaded));
 }
 
-bool UWorldPartitionEditorState::RestoreState() const
+UEditorState::FOperationResult UWorldPartitionEditorState::RestoreState() const
 {
 	UWorld* CurrentWorld = GetStateWorld();
-	if (CurrentWorld && CurrentWorld->IsPartitionedWorld())
+	if (!ensure(CurrentWorld) || !CurrentWorld->IsPartitionedWorld())
 	{
-		if (GetDefault<UWorldPartitionEditorSettings>()->GetEnableLoadingInEditor())
-		{
-			auto GetRegionHash = [](const FBox& LoadedRegion)
-			{
-				// Round region's min/max to avoid precision issues.
-				FInt64Vector LoadedRegionMin(FMath::RoundToInt(LoadedRegion.Min.X), FMath::RoundToInt(LoadedRegion.Min.Y), FMath::RoundToInt(LoadedRegion.Min.Z));
-				FInt64Vector LoadedRegionMax(FMath::RoundToInt(LoadedRegion.Max.X), FMath::RoundToInt(LoadedRegion.Max.Y), FMath::RoundToInt(LoadedRegion.Max.Z));
-				return HashCombineFast(GetTypeHash(LoadedRegionMin), GetTypeHash(LoadedRegionMax));
-			};
-
-			// Grab existing loaded regions to avoid creating duplicate regions
-			TSet<uint32> LoadedEditorRegionsHashes;
-			for (const FBox& LoadedRegion : CurrentWorld->GetWorldPartition()->GetUserLoadedEditorRegions())
-			{
-				LoadedEditorRegionsHashes.FindOrAdd(GetRegionHash(LoadedRegion));
-			}
+		return FOperationResult(FOperationResult::Skipped, LOCTEXT("RestoreStateSkipped_WorldIsNotPartitioned", "World is not partitioned"));
+	}
 
-			// Loaded regions
-			for (const FBox& LoadedRegion : LoadedEditorRegions)
-			{
-				if (LoadedRegion.IsValid)
-				{
-					bool bAlreadyInSet = false;
-					LoadedEditorRegionsHashes.FindOrAdd(GetRegionHash(LoadedRegion), &bAlreadyInSet);
-					if (!bAlreadyInSet)
-					{
-						UWorldPartitionEditorLoaderAdapter* EditorLoaderAdapter = CurrentWorld->GetWorldPartition()->CreateEditorLoaderAdapter<FLoaderAdapterShape>(CurrentWorld, LoadedRegion, TEXT("Loaded Region"));
-						IWorldPartitionActorLoaderInterface::ILoaderAdapter* LoaderAdapter = EditorLoaderAdapter->GetLoaderAdapter();
-						check(LoaderAdapter);
-						LoaderAdapter->SetUserCreated(true);
-						LoaderAdapter->Load();
-					}
-				}
-			}
+	if (!GetDefault<UWorldPartitionEditorSettings>()->GetEnableLoadingInEditor())
+	{
+		return FOperationResult(FOperationResult::Skipped, LOCTEXT("RestoreStateSkipped_LoadingInEditorDisabled", "Loading in editor is disabled"));
+	}
+
+	auto GetRegionHash = [](const FBox& LoadedRegion)
+	{
+		// Round region's min/max to avoid precision issues.
+		FInt64Vector LoadedRegionMin(FMath::RoundToInt(LoadedRegion.Min.X), FMath::RoundToInt(LoadedRegion.Min.Y), FMath::RoundToInt(LoadedRegion.Min.Z));
+		FInt64Vector LoadedRegionMax(FMath::RoundToInt(LoadedRegion.Max.X), FMath::RoundToInt(LoadedRegion.Max.Y), FMath::RoundToInt(LoadedRegion.Max.Z));
+		return HashCombineFast(GetTypeHash(LoadedRegionMin), GetTypeHash(LoadedRegionMax));
+	};
+
+	// Grab existing loaded regions to avoid creating duplicate regions
+	TSet<uint32> LoadedEditorRegionsHashes;
+	for (const FBox& LoadedRegion : CurrentWorld->GetWorldPartition()->GetUserLoadedEditorRegions())
+	{
+		LoadedEditorRegionsHashes.FindOrAdd(GetRegionHash(LoadedRegion));
+	}
 
-			// Location volumes
-			for (const FName& EditorLoadedLocationVolume : LoadedEditorLocationVolumes)
+	// Loaded regions
+	int32 NbRegionsLoaded = 0;
+	for (const FBox& LoadedRegion : LoadedEditorRegions)
+	{
+		if (LoadedRegion.IsValid)
+		{
+			bool bAlreadyInSet = false;
+			LoadedEditorRegionsHashes.FindOrAdd(GetRegionHash(LoadedRegion), &bAlreadyInSet);
+			if (!bAlreadyInSet)
 			{
-				if (ALocationVolume* LocationVolume = FindObject<ALocationVolume>(CurrentWorld->PersistentLevel, *EditorLoadedLocationVolume.ToString()))
-				{
-					LocationVolume->Load();
-				}
+				UWorldPartitionEditorLoaderAdapter* EditorLoaderAdapter = CurrentWorld->GetWorldPartition()->CreateEditorLoaderAdapter<FLoaderAdapterShape>(CurrentWorld, LoadedRegion, TEXT("Loaded Region"));
+				IWorldPartitionActorLoaderInterface::ILoaderAdapter* LoaderAdapter = EditorLoaderAdapter->GetLoaderAdapter();
+				check(LoaderAdapter);
+				LoaderAdapter->SetUserCreated(true);
+				LoaderAdapter->Load();
+				NbRegionsLoaded++;
 			}
 		}
 	}
 
-	return true;
+	// Location volumes
+	int32 NbVolumesLoaded = 0;
+	for (const FName& EditorLoadedLocationVolume : LoadedEditorLocationVolumes)
+	{
+		if (ALocationVolume* LocationVolume = FindObject<ALocationVolume>(CurrentWorld->PersistentLevel, *EditorLoadedLocationVolume.ToString()))
+		{
+			LocationVolume->Load();
+			NbVolumesLoaded++;
+		}
+	}
+	
+	return FOperationResult(FOperationResult::Success, FText::Format(LOCTEXT("RestoreStateSuccess", "Loaded {0} regions and {1} location volumes"), NbRegionsLoaded, NbVolumesLoaded));
 }
 
 #undef LOCTEXT_NAMESPACE
diff --git a/Engine/Source/Editor/WorldPartitionEditor/Public/WorldPartition/WorldPartitionEditorState.h b/Engine/Source/Editor/WorldPartitionEditor/Public/WorldPartition/WorldPartitionEditorState.h
--- a/Engine/Source/Editor/WorldPartitionEditor/Public/WorldPartition/WorldPartitionEditorState.h
+++ b/Engine/Source/Editor/WorldPartitionEditor/Public/WorldPartition/WorldPartitionEditorState.h
@@ -13,8 +13,8 @@ class UWorldPartitionEditorState : public UWorldDependantEditorState
 	//~ Begin UEditorState Interface
 	virtual FText GetCategoryText() const override;
 private:
-	virtual bool CaptureState() override;
-	virtual bool RestoreState() const override;
+	virtual FOperationResult CaptureState() override;
+	virtual FOperationResult RestoreState() const override;
 	//~ End UEditorState Interface
 
 private:
