World Bookmark - Fix regions & location volumes not being saved to bookmarks

diff --git a/Engine/Source/Editor/WorldPartitionEditor/Private/WorldPartition/WorldPartitionEditorState.cpp b/Engine/Source/Editor/WorldPartitionEditor/Private/WorldPartition/WorldPartitionEditorState.cpp
--- a/Engine/Source/Editor/WorldPartitionEditor/Private/WorldPartition/WorldPartitionEditorState.cpp
+++ b/Engine/Source/Editor/WorldPartitionEditor/Private/WorldPartition/WorldPartitionEditorState.cpp
@@ -40,7 +40,6 @@ UEditorState::FOperationResult UWorldPartitionEditorState::CaptureState()
 	};
 
 	// Loaded regions
-	int32 NbRegionsLoaded = 0;
 	TSet<uint32> LoadedEditorRegionsHashes;
 	for (const FBox& LoadedRegion : CurrentWorld->GetWorldPartition()->GetUserLoadedEditorRegions())
 	{
@@ -53,7 +52,6 @@ UEditorState::FOperationResult UWorldPartitionEditorState::CaptureState()
 	}
 
 	// Loaded location volumes
-	int32 NbVolumesLoaded = 0;
 	for (FActorDescContainerInstanceCollection::TConstIterator<> Iterator(CurrentWorld->GetWorldPartition()); Iterator; ++Iterator)
 	{
 		if (ALocationVolume* LocationVolume = Cast<ALocationVolume>(Iterator->GetActor()); IsValid(LocationVolume))
@@ -70,12 +68,12 @@ UEditorState::FOperationResult UWorldPartitionEditorState::CaptureState()
 		}
 	}
 
-	if (NbRegionsLoaded == 0 && NbVolumesLoaded == 0)
+	if (LoadedEditorRegions.IsEmpty() && LoadedEditorLocationVolumes.IsEmpty())
 	{
 		return FOperationResult(FOperationResult::Skipped, LOCTEXT("CaptureStateSkipped_NoLoadedRegionsOrVolumes", "No regions or location volumes currently loaded"));
 	}
 
-	return FOperationResult(FOperationResult::Success, FText::Format(LOCTEXT("CaptureStateSuccess", "{0} regions and {1} location volumes currently loaded"), NbRegionsLoaded, NbVolumesLoaded));
+	return FOperationResult(FOperationResult::Success, FText::Format(LOCTEXT("CaptureStateSuccess", "{0} regions and {1} location volumes currently loaded"), LoadedEditorRegions.Num(), LoadedEditorLocationVolumes.Num()));
 }
 
 UEditorState::FOperationResult UWorldPartitionEditorState::RestoreState() const
