Adapted existing 'Repair' option in the outliner to also detect and repair a previously deleted AWorldDataLayers (will save both the newly re-created WorldDataLayers actor and the map referencing it). Auto-repair was already in place, but nothing notified the user to force save the packages (AWorldDataLayers and the World referencing it).

diff --git a/Engine/Source/Editor/SceneOutliner/Private/ActorBrowsingMode.cpp b/Engine/Source/Editor/SceneOutliner/Private/ActorBrowsingMode.cpp
--- a/Engine/Source/Editor/SceneOutliner/Private/ActorBrowsingMode.cpp
+++ b/Engine/Source/Editor/SceneOutliner/Private/ActorBrowsingMode.cpp
@@ -46,9 +46,6 @@
 #include "WorldPartition/ActorDescContainerInstance.h"
 #include "WorldPartition/WorldPartitionActorDescInstance.h"
 #include "Subsystems/ActorEditorContextSubsystem.h"
-#include "ISourceControlModule.h"
-#include "ISourceControlProvider.h"
-#include "SourceControlOperations.h"
 #include "EditorLevelUtils.h"
 #include "EditorViewportCommands.h"
 #include "Misc/ScopedSlowTask.h"
@@ -1383,22 +1380,9 @@ bool FActorBrowsingMode::CanPaste() const
 
 bool FActorBrowsingMode::HasErrors() const
 {
-	if (RepresentingWorld.IsValid() && !bRepresentingWorldGameWorld && bRepresentingWorldPartitionedWorld)
+	if (RepresentingWorld.IsValid() && !bRepresentingWorldGameWorld && bRepresentingWorldPartitionedWorld && WorldPartitionEditorModule)
 	{
-		if (UWorldPartition* const WorldPartition = RepresentingWorld->GetWorldPartition())
-		{
-			bool bHasErrors = 0;
-
-			WorldPartition->ForEachActorDescContainerInstance([&bHasErrors](UActorDescContainerInstance* ActorDescContainerInstance)
-			{
-				if (ActorDescContainerInstance->GetContainer()->HasInvalidActors())
-				{
-					bHasErrors = true;
-				}
-			});
-
-			return bHasErrors;
-		}
+		return WorldPartitionEditorModule->HasErrors(RepresentingWorld.Get());
 	}
 
 	return false;
@@ -1411,88 +1395,9 @@ FText FActorBrowsingMode::GetErrorsText() const
 
 void FActorBrowsingMode::RepairErrors() const
 {
-	if (RepresentingWorld.IsValid() && !bRepresentingWorldGameWorld && bRepresentingWorldPartitionedWorld)
+	if (RepresentingWorld.IsValid() && !bRepresentingWorldGameWorld && bRepresentingWorldPartitionedWorld && WorldPartitionEditorModule)
 	{
-		if (UWorldPartition* const WorldPartition = RepresentingWorld->GetWorldPartition())
-		{
-			ISourceControlModule& SourceControlModule = ISourceControlModule::Get();
-			ISourceControlProvider& SourceControlProvider = SourceControlModule.GetProvider();
-
-			TArray<FAssetData> InvalidActorAssets;
-			WorldPartition->ForEachActorDescContainerInstance([&InvalidActorAssets](UActorDescContainerInstance* ActorDescContainerInstance)
-			{
-				for (const FAssetData& InvalidActor : ActorDescContainerInstance->GetContainer()->GetInvalidActors())
-				{
-					InvalidActorAssets.Add(InvalidActor);
-				}
-				ActorDescContainerInstance->GetContainer()->ClearInvalidActors();
-			});
-
-			FlushAsyncLoading();
-
-			TArray<FString> ActorFilesToDelete;
-			TArray<FString> ActorFilesToRevert;
-			{
-				FScopedSlowTask SlowTask(InvalidActorAssets.Num(), LOCTEXT("UpdatingSourceControlStatus", "Updating source control status..."));
-				SlowTask.MakeDialogDelayed(1.0f);
-
-				for (const FAssetData& InvalidActorAsset : InvalidActorAssets)
-				{
-					FPackagePath PackagePath;
-					if (FPackagePath::TryFromPackageName(InvalidActorAsset.PackageName, PackagePath))
-					{
-						if (UPackage* ExistingPackage = FindObject<UPackage>(nullptr, *InvalidActorAsset.PackageName.ToString()))
-						{
-							ResetLoaders(ExistingPackage);
-						}
-
-						const FString ActorFile = PackagePath.GetLocalFullPath();
-						FSourceControlStatePtr SCState = SourceControlProvider.GetState(ActorFile, EStateCacheUsage::ForceUpdate);
-
-						if (SCState.IsValid() && SCState->IsSourceControlled())
-						{
-							if (SCState->IsAdded())
-							{
-								ActorFilesToRevert.Add(ActorFile);
-							}
-							else
-							{
-								if (SCState->IsCheckedOut())
-								{
-									ActorFilesToRevert.Add(ActorFile);
-								}
-
-								ActorFilesToDelete.Add(ActorFile);
-							}
-						}
-						else
-						{
-							IFileManager::Get().Delete(*ActorFile, false, true);
-						}
-					}
-
-					SlowTask.EnterProgressFrame(1);
-				}
-			}
-
-			if (ActorFilesToRevert.Num() || ActorFilesToDelete.Num())
-			{
-				FScopedSlowTask SlowTask(ActorFilesToRevert.Num() + ActorFilesToDelete.Num(), LOCTEXT("DeletingFiles", "Deleting files..."));
-				SlowTask.MakeDialogDelayed(1.0f);
-
-				if (ActorFilesToRevert.Num())
-				{
-					SourceControlProvider.Execute(ISourceControlOperation::Create<FRevert>(), ActorFilesToRevert);
-					SlowTask.EnterProgressFrame(ActorFilesToRevert.Num());
-				}
-	
-				if (ActorFilesToDelete.Num())
-				{
-					SourceControlProvider.Execute(ISourceControlOperation::Create<FDelete>(), ActorFilesToDelete);
-					SlowTask.EnterProgressFrame(ActorFilesToDelete.Num());
-				}
-			}
-		}
+		WorldPartitionEditorModule->RepairErrors(RepresentingWorld.Get());
 	}
 }
 
diff --git a/Engine/Source/Editor/WorldPartitionEditor/Private/WorldPartitionEditorModule.cpp b/Engine/Source/Editor/WorldPartitionEditor/Private/WorldPartitionEditorModule.cpp
--- a/Engine/Source/Editor/WorldPartitionEditor/Private/WorldPartitionEditorModule.cpp
+++ b/Engine/Source/Editor/WorldPartitionEditor/Private/WorldPartitionEditorModule.cpp
@@ -57,6 +57,9 @@
 #include "Framework/Docking/LayoutExtender.h"
 #include "Widgets/Docking/SDockTab.h"
 #include "Filters/CustomClassFilterData.h"
+#include "ISourceControlModule.h"
+#include "ISourceControlProvider.h"
+#include "SourceControlOperations.h"
 
 #include "Styling/AppStyle.h"
 #include "WorldPartition/ContentBundle/SContentBundleBrowser.h"
@@ -1018,6 +1021,145 @@ void FWorldPartitionEditorModule::RegisterWorldPartitionLayout(FLayoutExtender&
 	Extender.ExtendLayout(FTabId("LevelEditorSelectionDetails"), ELayoutExtensionPosition::After, FTabManager::FTab(WorldPartitionEditorTabId, ETabState::ClosedTab));
 }
 
+static bool ShouldRepairWorldDataLayers(UWorld* World)
+{
+	AWorldDataLayers* WorldDataLayers = World ? World->PersistentLevel->GetWorldDataLayers() : nullptr;
+	if (WorldDataLayers &&
+		WorldDataLayers->GetPackage()->HasAnyPackageFlags(PKG_NewlyCreated) &&
+		!World->GetPackage()->HasAnyPackageFlags(PKG_NewlyCreated) &&
+		!FPackageName::DoesPackageExist(WorldDataLayers->GetPackage()->GetName()) &&
+		FPackageName::DoesPackageExist(World->GetPackage()->GetName()))
+	{
+		return true;
+	}
+	return false;
+}
+
+bool FWorldPartitionEditorModule::HasErrors(UWorld* World) const
+{
+	if (UWorldPartition* const WorldPartition = World ? World->GetWorldPartition() : nullptr)
+	{
+		if (ShouldRepairWorldDataLayers(World))
+		{
+			return true;
+		}
+
+		bool bHasErrors = false;
+
+		WorldPartition->ForEachActorDescContainerInstanceBreakable([&bHasErrors](UActorDescContainerInstance* ActorDescContainerInstance)
+		{
+			if (ActorDescContainerInstance->GetContainer()->HasInvalidActors())
+			{
+				bHasErrors = true;
+			}
+			return !bHasErrors;
+		});
+
+		return bHasErrors;
+	}
+
+	return false;
+}
+
+void FWorldPartitionEditorModule::RepairErrors(UWorld* World)
+{
+	if (UWorldPartition* const WorldPartition = World ? World->GetWorldPartition() : nullptr)
+	{
+		if (ShouldRepairWorldDataLayers(World))
+		{
+			TSet<UPackage*> PackagesToSave;
+			PackagesToSave.Add(World->PersistentLevel->GetWorldDataLayers()->GetPackage());
+			PackagesToSave.Add(World->GetPackage());
+			if (PackagesToSave.Num())
+			{
+				FEditorFileUtils::FPromptForCheckoutAndSaveParams SaveParams;
+				SaveParams.bCheckDirty = false;
+				SaveParams.bPromptToSave = false;
+				SaveParams.bIsExplicitSave = true;
+				FEditorFileUtils::PromptForCheckoutAndSave(PackagesToSave.Array(), SaveParams);
+			}
+		}
+
+		ISourceControlModule& SourceControlModule = ISourceControlModule::Get();
+		ISourceControlProvider& SourceControlProvider = SourceControlModule.GetProvider();
+
+		TArray<FAssetData> InvalidActorAssets;
+		WorldPartition->ForEachActorDescContainerInstance([&InvalidActorAssets](UActorDescContainerInstance* ActorDescContainerInstance)
+		{
+			for (const FAssetData& InvalidActor : ActorDescContainerInstance->GetContainer()->GetInvalidActors())
+			{
+				InvalidActorAssets.Add(InvalidActor);
+			}
+			ActorDescContainerInstance->GetContainer()->ClearInvalidActors();
+		});
+
+		FlushAsyncLoading();
+
+		TArray<FString> ActorFilesToDelete;
+		TArray<FString> ActorFilesToRevert;
+		{
+			FScopedSlowTask SlowTask(InvalidActorAssets.Num(), LOCTEXT("UpdatingSourceControlStatus", "Updating source control status..."));
+			SlowTask.MakeDialogDelayed(1.0f);
+
+			for (const FAssetData& InvalidActorAsset : InvalidActorAssets)
+			{
+				FPackagePath PackagePath;
+				if (FPackagePath::TryFromPackageName(InvalidActorAsset.PackageName, PackagePath))
+				{
+					if (UPackage* ExistingPackage = FindObject<UPackage>(nullptr, *InvalidActorAsset.PackageName.ToString()))
+					{
+						ResetLoaders(ExistingPackage);
+					}
+
+					const FString ActorFile = PackagePath.GetLocalFullPath();
+					FSourceControlStatePtr SCState = SourceControlProvider.GetState(ActorFile, EStateCacheUsage::ForceUpdate);
+
+					if (SCState.IsValid() && SCState->IsSourceControlled())
+					{
+						if (SCState->IsAdded())
+						{
+							ActorFilesToRevert.Add(ActorFile);
+						}
+						else
+						{
+							if (SCState->IsCheckedOut())
+							{
+								ActorFilesToRevert.Add(ActorFile);
+							}
+
+							ActorFilesToDelete.Add(ActorFile);
+						}
+					}
+					else
+					{
+						IFileManager::Get().Delete(*ActorFile, false, true);
+					}
+				}
+
+				SlowTask.EnterProgressFrame(1);
+			}
+		}
+
+		if (ActorFilesToRevert.Num() || ActorFilesToDelete.Num())
+		{
+			FScopedSlowTask SlowTask(ActorFilesToRevert.Num() + ActorFilesToDelete.Num(), LOCTEXT("DeletingFiles", "Deleting files..."));
+			SlowTask.MakeDialogDelayed(1.0f);
+
+			if (ActorFilesToRevert.Num())
+			{
+				SourceControlProvider.Execute(ISourceControlOperation::Create<FRevert>(), ActorFilesToRevert);
+				SlowTask.EnterProgressFrame(ActorFilesToRevert.Num());
+			}
+
+			if (ActorFilesToDelete.Num())
+			{
+				SourceControlProvider.Execute(ISourceControlOperation::Create<FDelete>(), ActorFilesToDelete);
+				SlowTask.EnterProgressFrame(ActorFilesToDelete.Num());
+			}
+		}
+	}
+}
+
 FString UWorldPartitionConvertOptions::ToCommandletArgs() const
 {
 	TStringBuilder<1024> CommandletArgsBuilder;
diff --git a/Engine/Source/Editor/WorldPartitionEditor/Public/WorldPartitionEditorModule.h b/Engine/Source/Editor/WorldPartitionEditor/Public/WorldPartitionEditorModule.h
--- a/Engine/Source/Editor/WorldPartitionEditor/Public/WorldPartitionEditorModule.h
+++ b/Engine/Source/Editor/WorldPartitionEditor/Public/WorldPartitionEditorModule.h
@@ -92,6 +92,9 @@ public:
 
 	virtual bool IsHLODInEditorAllowed(UWorld* InWorld, FText* OutDisallowedReason) const override;
 
+	virtual bool HasErrors(UWorld* World) const override;
+	virtual void RepairErrors(UWorld* World) override;
+
 	/**
      * Writes various HLOD stats to a file.
      */ 
diff --git a/Engine/Source/Runtime/Engine/Public/WorldPartition/IWorldPartitionEditorModule.h b/Engine/Source/Runtime/Engine/Public/WorldPartition/IWorldPartitionEditorModule.h
--- a/Engine/Source/Runtime/Engine/Public/WorldPartition/IWorldPartitionEditorModule.h
+++ b/Engine/Source/Runtime/Engine/Public/WorldPartition/IWorldPartitionEditorModule.h
@@ -69,6 +69,12 @@ public:
 
 	virtual bool IsHLODInEditorAllowed(UWorld* InWorld, FText* OutDisallowedReason = nullptr) const = 0;
 
+	/** Returns if the partitioned world contains errors. */
+	virtual bool HasErrors(UWorld* World) const = 0;
+
+	/** Repairs detected errors in the partitioned world. */
+	virtual void RepairErrors(UWorld* World) = 0;
+
 	/**
      * Parameters for the WriteHLODStats() method.
      */	
