Disable EnableStreamingGenerationLogOnPIE by default, keep it enabled internally.

diff --git a/Engine/Source/Editor/WorldPartitionEditor/Private/WorldPartition/WorldPartitionEditorSettings.cpp b/Engine/Source/Editor/WorldPartitionEditor/Private/WorldPartition/WorldPartitionEditorSettings.cpp
--- a/Engine/Source/Editor/WorldPartitionEditor/Private/WorldPartition/WorldPartitionEditorSettings.cpp
+++ b/Engine/Source/Editor/WorldPartitionEditor/Private/WorldPartition/WorldPartitionEditorSettings.cpp
@@ -11,7 +11,7 @@ UWorldPartitionEditorSettings::UWorldPartitionEditorSettings()
 	InstancedFoliageGridSize = 25600;
 	MinimapLowQualityWorldUnitsPerPixelThreshold = 12800;
 	bEnableLoadingInEditor = true;
-	bEnableStreamingGenerationLogOnPIE = true;
+	bEnableStreamingGenerationLogOnPIE = false;
 	bShowHLODsInEditor = true;
 	bShowHLODsOverLoadedRegions = false;
 	HLODMinDrawDistance = 12800;
