[HLOD] Make -DumpStats available with cmdlet.

#tests wp.editor.hlod.dumpstats == dumpstats from cmdline
[RN]

diff --git a/Engine/Source/Editor/WorldPartitionEditor/Private/WorldPartition/HLOD/HLODEditorSubsystem.cpp b/Engine/Source/Editor/WorldPartitionEditor/Private/WorldPartition/HLOD/HLODEditorSubsystem.cpp
--- a/Engine/Source/Editor/WorldPartitionEditor/Private/WorldPartition/HLOD/HLODEditorSubsystem.cpp
+++ b/Engine/Source/Editor/WorldPartitionEditor/Private/WorldPartition/HLOD/HLODEditorSubsystem.cpp
@@ -77,6 +77,11 @@ UWorldPartitionHLODEditorSubsystem::~UWorldPartitionHLODEditorSubsystem()
 
 bool UWorldPartitionHLODEditorSubsystem::IsHLODInEditorEnabled()
 {
+	if (IsRunningCommandlet())
+	{
+		return false;
+	}
+	
 	const IWorldPartitionEditorModule* WorldPartitionEditorModule = FModuleManager::GetModulePtr<IWorldPartitionEditorModule>("WorldPartitionEditor");
 	const bool bShowHLODsInEditorForWorld = WorldPartitionEditorModule && WorldPartitionEditorModule->IsHLODInEditorAllowed(GetWorld());
 	const bool bShowHLODsInEditorUserSetting = WorldPartitionEditorModule && WorldPartitionEditorModule->GetShowHLODsInEditor();
@@ -86,7 +91,7 @@ bool UWorldPartitionHLODEditorSubsystem::IsHLODInEditorEnabled()
 
 bool UWorldPartitionHLODEditorSubsystem::DoesSupportWorldType(const EWorldType::Type WorldType) const
 {
-	return WorldType == EWorldType::Editor && !IsRunningCommandlet();
+	return WorldType == EWorldType::Editor;
 }
 
 void UWorldPartitionHLODEditorSubsystem::Initialize(FSubsystemCollectionBase& Collection)
