Added option in World Partition View to show/hide loading regions

diff --git a/Engine/Source/Editor/EditorStyle/Private/StarshipStyle.cpp b/Engine/Source/Editor/EditorStyle/Private/StarshipStyle.cpp
--- a/Engine/Source/Editor/EditorStyle/Private/StarshipStyle.cpp
+++ b/Engine/Source/Editor/EditorStyle/Private/StarshipStyle.cpp
@@ -2093,6 +2093,7 @@ void FStarshipEditorStyle::FStyle::SetupWorldPartitionStyles()
 	Set("WorldPartition.ShowGrid", new IMAGE_BRUSH_SVG("Starship/EditorViewport/grid", Icon16x16));
 	Set("WorldPartition.ShowMiniMap", new IMAGE_BRUSH_SVG("Starship/AssetIcons/Texture2D_16", Icon16x16)); 
 	Set("WorldPartition.ShowCoords", new IMAGE_BRUSH_SVG("WorldPartition/Coordinate_16", Icon16x16));
+	Set("WorldPartition.ShowLoadingRegions", new IMAGE_BRUSH_SVG("Starship/Common/Volumes", Icon16x16));
 	Set("WorldPartition.ShowMouseCoords", new IMAGE_BRUSH_SVG("WorldPartition/MouseCoordinates_16", Icon16x16));
 	Set("WorldPartition.FocusSelection", new IMAGE_BRUSH_SVG("WorldPartition/FramingSelection_20", Icon16x16));
 	Set("WorldPartition.FocusLoadedRegions", new IMAGE_BRUSH_SVG("WorldPartition/FrameRegions_20", Icon16x16));
diff --git a/Engine/Source/Editor/WorldPartitionEditor/Private/WorldPartition/SWorldPartitionEditorGrid2D.cpp b/Engine/Source/Editor/WorldPartitionEditor/Private/WorldPartition/SWorldPartitionEditorGrid2D.cpp
--- a/Engine/Source/Editor/WorldPartitionEditor/Private/WorldPartition/SWorldPartitionEditorGrid2D.cpp
+++ b/Engine/Source/Editor/WorldPartitionEditor/Private/WorldPartition/SWorldPartitionEditorGrid2D.cpp
@@ -247,6 +247,7 @@ void SWorldPartitionEditorGrid2D::FEditorCommands::RegisterCommands()
 	UI_COMMAND(ShowMiniMap, "Minimap", "Show the minimap texture.", EUserInterfaceActionType::ToggleButton, FInputChord());
 	UI_COMMAND(ShowCoords, "Coordinates", "Show grid cell coordinates, you might need to zoom closer to see them.", EUserInterfaceActionType::ToggleButton, FInputChord());
 	UI_COMMAND(ShowMouseCoords, "Mouse Coordinates", "Show mouse world coordinates.", EUserInterfaceActionType::ToggleButton, FInputChord());
+	UI_COMMAND(ShowLoadingRegions, "Loading Regions", "Show loading regions.", EUserInterfaceActionType::ToggleButton, FInputChord());
 
 	UI_COMMAND(FocusSelection, "Focus Selection", "Focus Selection.", EUserInterfaceActionType::ToggleButton, FInputChord());
 	UI_COMMAND(FocusLoadedRegions, "Focus Loaded Regions", "Focus Loaded Regions.", EUserInterfaceActionType::ToggleButton, FInputChord());
@@ -407,6 +408,7 @@ TSharedRef<SWidget> SWorldPartitionEditorGrid2D::SToolBar::GenerateShowMenu() co
 	Section.AddMenuEntry(Commands.ShowMiniMap);
 	Section.AddMenuEntry(Commands.ShowCoords);
 	Section.AddMenuEntry(Commands.ShowMouseCoords);
+	Section.AddMenuEntry(Commands.ShowLoadingRegions);
 
 	return UToolMenus::Get()->GenerateWidget(MenuName, FToolMenuContext(CommandList));
 }
@@ -440,6 +442,7 @@ SWorldPartitionEditorGrid2D::SWorldPartitionEditorGrid2D()
 	, bShowActors(false)
 	, bShowHLODActors(false)
 	, bShowGrid(true)
+	, bShowLoadingRegions(true)
 	, bShowMiniMap(true)
 	, bShowMouseCoords(false)
 	, bFollowPlayerInPIE(false)
@@ -620,6 +623,7 @@ void SWorldPartitionEditorGrid2D::BindCommands()
 	CommandList->MapAction(Commands.ShowMiniMap, FExecuteAction::CreateLambda([this]() { bShowMiniMap = !bShowMiniMap; }), FCanExecuteAction(), FIsActionChecked::CreateLambda([this]() { return bShowMiniMap; }));
 	CommandList->MapAction(Commands.ShowCoords, FExecuteAction::CreateLambda([this]() { GetMutableDefault<UWorldPartitionEditorPerProjectUserSettings>()->SetShowCellCoords(!GetMutableDefault<UWorldPartitionEditorPerProjectUserSettings>()->GetShowCellCoords()); }), FCanExecuteAction(), FIsActionChecked::CreateLambda([this]() { return GetMutableDefault<UWorldPartitionEditorPerProjectUserSettings>()->GetShowCellCoords(); }), FIsActionButtonVisible::CreateLambda([this]() { return (GetWorldPartition() && GetWorldPartition()->IsStreamingEnabled()); }));
 	CommandList->MapAction(Commands.ShowMouseCoords, FExecuteAction::CreateLambda([this]() { bShowMouseCoords = !bShowMouseCoords; }), FCanExecuteAction(), FIsActionChecked::CreateLambda([this]() { return bShowMouseCoords; }));
+	CommandList->MapAction(Commands.ShowLoadingRegions, FExecuteAction::CreateLambda([this]() { bShowLoadingRegions = !bShowLoadingRegions; }), FCanExecuteAction(), FIsActionChecked::CreateLambda([this]() { return bShowLoadingRegions; }));
 
 	// Quick Actions
 	CommandList->MapAction(Commands.FocusSelection, FExecuteAction::CreateSP(this, &SWorldPartitionEditorGrid2D::FocusSelection), FCanExecuteAction::CreateLambda(CanFocusSelection));
@@ -1480,7 +1484,7 @@ uint32 SWorldPartitionEditorGrid2D::PaintActors(const FGeometry& AllottedGeometr
 		}
 	};
 
-	if (ShownLoaderInterfaces.Num())
+	if (bShowLoadingRegions && ShownLoaderInterfaces.Num())
 	{
 		TArray<FVector2D> LinePoints;
 		LinePoints.SetNum(5);
diff --git a/Engine/Source/Editor/WorldPartitionEditor/Public/WorldPartition/SWorldPartitionEditorGrid2D.h b/Engine/Source/Editor/WorldPartitionEditor/Public/WorldPartition/SWorldPartitionEditorGrid2D.h
--- a/Engine/Source/Editor/WorldPartitionEditor/Public/WorldPartition/SWorldPartitionEditorGrid2D.h
+++ b/Engine/Source/Editor/WorldPartitionEditor/Public/WorldPartition/SWorldPartitionEditorGrid2D.h
@@ -47,6 +47,7 @@ protected:
 		TSharedPtr<FUICommandInfo> ShowMiniMap;
 		TSharedPtr<FUICommandInfo> ShowCoords;
 		TSharedPtr<FUICommandInfo> ShowMouseCoords;
+		TSharedPtr<FUICommandInfo> ShowLoadingRegions;
 
 		// Quick Actions
 		TSharedPtr<FUICommandInfo> FocusSelection;
@@ -169,6 +170,7 @@ protected:
 	bool bShowActors;
 	bool bShowHLODActors;
 	bool bShowGrid;
+	bool bShowLoadingRegions;
 	bool bShowMiniMap;
 	bool bShowMouseCoords;
 	bool bFollowPlayerInPIE;
