Enable virtual texture feedback for Slate UI renderer.
Allow virtual textures to be used in UI material domain.
Remove CPU driven virtual texture feedback for world partition minimap. This now happens "for free" for the Slate widget.

diff --git a/Engine/Source/Editor/WorldPartitionEditor/Private/WorldPartition/SWorldPartitionEditorGrid2D.cpp b/Engine/Source/Editor/WorldPartitionEditor/Private/WorldPartition/SWorldPartitionEditorGrid2D.cpp
--- a/Engine/Source/Editor/WorldPartitionEditor/Private/WorldPartition/SWorldPartitionEditorGrid2D.cpp
+++ b/Engine/Source/Editor/WorldPartitionEditor/Private/WorldPartition/SWorldPartitionEditorGrid2D.cpp
@@ -2087,34 +2087,6 @@ int32 SWorldPartitionEditorGrid2D::PaintMinimap(const FGeometry& AllottedGeometr
 				USlateThemeManager::Get().GetColor(EStyleColor::Black)
 			);
 		}
-
-		if (Texture2D->IsCurrentlyVirtualTextured())
-		{
-			FVirtualTexture2DResource* VTResource = static_cast<FVirtualTexture2DResource*>(Texture2D->GetResource());
-			const FVector2D ViewportSize = AllottedGeometry.GetLocalSize();
-			const FVector2D ScreenSpaceSize = MinimapBounds.GetSize();
-			const FVector2D ViewportPosition = MinimapBounds.Min;
-
-			FBox2D UVRegion = WorldMiniMapBrush.GetUVRegion();
-			const FVector2D UV0 = UVRegion.Min;
-			const FVector2D UV1 = UVRegion.Max;
-
-			const ERHIFeatureLevel::Type InFeatureLevel = GMaxRHIFeatureLevel;
-			const int32 MipLevel = -1;
-
-			UE::RenderCommandPipe::FSyncScope SyncScope;
-
-			ENQUEUE_RENDER_COMMAND(MakeTilesResident)(
-				[InFeatureLevel, VTResource, ScreenSpaceSize, ViewportPosition, ViewportSize, UV0, UV1, MipLevel](FRHICommandListImmediate& RHICmdList)
-			{
-				// AcquireAllocatedVT() must happen on render thread
-				IAllocatedVirtualTexture* AllocatedVT = VTResource->AcquireAllocatedVT();
-
-				IRendererModule& RenderModule = GetRendererModule();
-				RenderModule.RequestVirtualTextureTiles(AllocatedVT, ScreenSpaceSize, ViewportPosition, ViewportSize, UV0, UV1, MipLevel);
-				RenderModule.LoadPendingVirtualTextureTiles(RHICmdList, InFeatureLevel);
-			});
-		}
 	}
 
 	return LayerId;
diff --git a/Engine/Source/Runtime/Engine/Private/Materials/HLSLMaterialTranslator.cpp b/Engine/Source/Runtime/Engine/Private/Materials/HLSLMaterialTranslator.cpp
--- a/Engine/Source/Runtime/Engine/Private/Materials/HLSLMaterialTranslator.cpp
+++ b/Engine/Source/Runtime/Engine/Private/Materials/HLSLMaterialTranslator.cpp
@@ -7222,6 +7222,7 @@ int32 FHLSLMaterialTranslator::TextureSample(
 			case MD_Surface:
 			case MD_DeferredDecal:
 			case MD_PostProcess:
+			case MD_UI:
 				break;
 			default:
 				return Errorf(TEXT("Sampling a virtual texture is not supported for the %s material domain."), *MaterialDomainString(Domain));
diff --git a/Engine/Source/Runtime/SlateRHIRenderer/Private/SlateRHIRenderer.cpp b/Engine/Source/Runtime/SlateRHIRenderer/Private/SlateRHIRenderer.cpp
--- a/Engine/Source/Runtime/SlateRHIRenderer/Private/SlateRHIRenderer.cpp
+++ b/Engine/Source/Runtime/SlateRHIRenderer/Private/SlateRHIRenderer.cpp
@@ -55,6 +55,7 @@
 #include "DataDrivenShaderPlatformInfo.h"
 #include "SlatePostProcessor.h"
 #include "Stats/ThreadIdleStats.h"
+#include "VT/VirtualTextureFeedbackResource.h"
 
 #if WITH_EDITORONLY_DATA
 #include "ShaderCompiler.h"
@@ -769,6 +770,12 @@ FSlateDrawWindowPassOutputs FSlateRHIRenderer::DrawWindow_RenderThread(FRDGBuild
 		FSlateBatchData& BatchData = WindowElementList.GetBatchData();
 		FSlateBatchData& BatchDataHDR = WindowElementList.GetBatchDataHDR();
 
+		const bool bRequiresVirtualTextureFeedback = BatchData.IsVirtualTextureFeedbackRequired() || BatchDataHDR.IsVirtualTextureFeedbackRequired();
+		if (bRequiresVirtualTextureFeedback)
+		{
+			VirtualTexture::BeginFeedback(GraphBuilder);
+		}
+
 		const FSlateElementsBuffers SlateElementsBuffers = BuildSlateElementsBuffers(GraphBuilder, BatchData);
 		const FSlateElementsBuffers SlateElementsBuffersHDR = BuildSlateElementsBuffers(GraphBuilder, BatchDataHDR);
 
@@ -900,6 +907,11 @@ FSlateDrawWindowPassOutputs FSlateRHIRenderer::DrawWindow_RenderThread(FRDGBuild
 			});
 		}
 
+		if (bRequiresVirtualTextureFeedback)
+		{
+			VirtualTexture::EndFeedback(GraphBuilder);
+		}
+
 		OnAddBackBufferReadyToPresentPassDelegate.Broadcast(GraphBuilder, *Inputs.Window, OutputTexture);
 
 		if (ScreenshotState.ViewportToCapture == &ViewportInfo)
