World Partition - HLOD: Modified default value for the HLOD in editor min draw distance, from 12800 to 0 (no minimum draw distance by default)

diff --git a/Engine/Source/Editor/WorldPartitionEditor/Private/WorldPartition/WorldPartitionEditorSettings.cpp b/Engine/Source/Editor/WorldPartitionEditor/Private/WorldPartition/WorldPartitionEditorSettings.cpp
--- a/Engine/Source/Editor/WorldPartitionEditor/Private/WorldPartition/WorldPartitionEditorSettings.cpp
+++ b/Engine/Source/Editor/WorldPartitionEditor/Private/WorldPartition/WorldPartitionEditorSettings.cpp
@@ -14,7 +14,7 @@ UWorldPartitionEditorSettings::UWorldPartitionEditorSettings()
 	bEnableStreamingGenerationLogOnPIE = false;
 	bShowHLODsInEditor = true;
 	bShowHLODsOverLoadedRegions = false;
-	HLODMinDrawDistance = 12800;
+	HLODMinDrawDistance = 0;
 	HLODMaxDrawDistance = 0;
 	bDisablePIE = false;
 	bDisableBugIt = false;
