World Partition - Fix issue occurring when running build commands from the editor menu on never saved worlds
* Background commmandlet job would fail (package is not yet on disk)
* Editor would load the engine template map once that failed build was completed

diff --git a/Engine/Source/Editor/WorldPartitionEditor/Private/WorldPartitionEditorModule.cpp b/Engine/Source/Editor/WorldPartitionEditor/Private/WorldPartitionEditorModule.cpp
--- a/Engine/Source/Editor/WorldPartitionEditor/Private/WorldPartitionEditorModule.cpp
+++ b/Engine/Source/Editor/WorldPartitionEditor/Private/WorldPartitionEditorModule.cpp
@@ -905,15 +905,25 @@ bool FWorldPartitionEditorModule::Build(const FRunBuilderParams& InParams)
 	FRunBuilderParams ParamsCopy(InParams);
 	OnPreExecuteCommandletEvent.Broadcast(ParamsCopy);
 
+	UWorld* World = ParamsCopy.World;
+	UPackage* WorldPackage = World->GetPackage();
+	
 	// Ask user to save dirty packages
 	if (!AskSaveDirtyPackages(/*bAskSaveContentPackages=*/true))
-	{
+	{			
 		return false;
 	}
 
+	// Validate that a newly created world was actually saved
+	if (WorldPackage->HasAnyPackageFlags(PKG_NewlyCreated))
+	{
+		FMessageDialog::Open(EAppMsgType::Ok, LOCTEXT("NewMap", "New world must be saved before performing this operation."));
+		return false;
+	}
+	
 	// Unload map if required
-	FString MapPackage = ParamsCopy.World->GetPackage()->GetName();
-	if (!UnloadCurrentMap(MapPackage))
+	FString WorldPackageName = WorldPackage->GetName();
+	if (!UnloadCurrentMap(WorldPackageName))
 	{
 		return false;
 	}
@@ -922,7 +932,7 @@ bool FWorldPartitionEditorModule::Build(const FRunBuilderParams& InParams)
 	GEditor->GetEditorSubsystem<UAssetEditorSubsystem>()->CloseAllAssetEditors();
 
 	TStringBuilder<512> CommandletArgsBuilder;
-	CommandletArgsBuilder.Append(MapPackage);
+	CommandletArgsBuilder.Append(WorldPackageName);
 	CommandletArgsBuilder.Append(" -run=WorldPartitionBuilderCommandlet -Builder=");
 	CommandletArgsBuilder.Append(ParamsCopy.BuilderClass->GetName());
 		
@@ -940,8 +950,8 @@ bool FWorldPartitionEditorModule::Build(const FRunBuilderParams& InParams)
 	FString CommandletArgs(CommandletArgsBuilder.ToString());
 	RunCommandletAsExternalProcess(CommandletArgs, OperationDescription, Result, bCancelled);
 
-	RescanAssets(MapPackage);
-	LoadMap(MapPackage);
+	RescanAssets(WorldPackageName);
+	LoadMap(WorldPackageName);
 
 	if (bCancelled)
 	{
