Fixed CIS warning

diff --git a/Engine/Source/Editor/WorldPartitionEditor/Private/WorldPartition/Customizations/WorldPartitionHLODLayerDetailsCustomization.cpp b/Engine/Source/Editor/WorldPartitionEditor/Private/WorldPartition/Customizations/WorldPartitionHLODLayerDetailsCustomization.cpp
--- a/Engine/Source/Editor/WorldPartitionEditor/Private/WorldPartition/Customizations/WorldPartitionHLODLayerDetailsCustomization.cpp
+++ b/Engine/Source/Editor/WorldPartitionEditor/Private/WorldPartition/Customizations/WorldPartitionHLODLayerDetailsCustomization.cpp
@@ -22,6 +22,10 @@ void FWorldPartitionHLODLayerDetailsCustomization::CustomizeDetails(IDetailLayou
 	UWorldPartition::WorldPartitionChangedEvent.AddSP(this, &FWorldPartitionHLODLayerDetailsCustomization::OnWorldPartitionChanged);
 
 	UHLODLayer* HLODLayer = !DetailLayoutBuilder->GetSelectedObjects().IsEmpty() ? Cast<UHLODLayer>(DetailLayoutBuilder->GetSelectedObjects()[0].Get()) : nullptr;
+	if (!HLODLayer)
+	{
+		return;
+	}
 		
 	TSharedPtr<IPropertyHandle> IsSpatiallyLoadedPropertyHandle = DetailLayoutBuilder->GetProperty(GET_MEMBER_NAME_CHECKED(UHLODLayerEditorOnlyData, bIsSpatiallyLoaded), UHLODLayerEditorOnlyData::StaticClass());
 	TSharedPtr<IPropertyHandle> CellSizePropertyHandle = DetailLayoutBuilder->GetProperty(GET_MEMBER_NAME_CHECKED(UHLODLayerEditorOnlyData, CellSize), UHLODLayerEditorOnlyData::StaticClass());
