World Partition - HLOD - When creating a new HLOD layer out of a template asset, make sure no parent HLOD Layer is assigned
* Creating a new layer and having it automatically use a parent layer does not make much sense

diff --git a/Engine/Source/Editor/WorldPartitionEditor/Private/WorldPartition/HLOD/HLODLayerFactory.cpp b/Engine/Source/Editor/WorldPartitionEditor/Private/WorldPartition/HLOD/HLODLayerFactory.cpp
--- a/Engine/Source/Editor/WorldPartitionEditor/Private/WorldPartition/HLOD/HLODLayerFactory.cpp
+++ b/Engine/Source/Editor/WorldPartitionEditor/Private/WorldPartition/HLOD/HLODLayerFactory.cpp
@@ -28,6 +28,9 @@ UObject* UHLODLayerFactory::FactoryCreateNew(UClass* Class, UObject* InParent, F
 		HLODLayer = DuplicateObject<UHLODLayer>(HLODLayerTemplate, InParent, Name);
 		HLODLayer->ClearFlags(RF_AllFlags);
 		HLODLayer->SetFlags(Flags);
+
+		// Make sure to never assign a parent HLOD layer to a newly created asset, even if the default template has one
+		HLODLayer->SetParentLayer(nullptr);
 	}
 	
 	if (!HLODLayer)
