World Partition - HLOD: Added an "Export HLOD Assets" option to HLOD actors
* Available for all HLODs which have private, internal assets in their OFPA packages (so this excludes Instanced HLODs)
* Exposed to BP

diff --git a/Engine/Source/Editor/WorldPartitionEditor/Private/WorldPartition/Customizations/WorldPartitionHLODDetailsCustomization.cpp b/Engine/Source/Editor/WorldPartitionEditor/Private/WorldPartition/Customizations/WorldPartitionHLODDetailsCustomization.cpp
--- a/Engine/Source/Editor/WorldPartitionEditor/Private/WorldPartition/Customizations/WorldPartitionHLODDetailsCustomization.cpp
+++ b/Engine/Source/Editor/WorldPartitionEditor/Private/WorldPartition/Customizations/WorldPartitionHLODDetailsCustomization.cpp
@@ -2,20 +2,30 @@
 
 #include "WorldPartitionHLODDetailsCustomization.h"
 
+#include "Algo/AnyOf.h"
+#include "Algo/Copy.h"
 #include "Algo/ForEach.h"
 #include "Algo/Transform.h"
 #include "ComponentRecreateRenderStateContext.h"
 #include "Components/StaticMeshComponent.h"
+#include "ContentBrowserModule.h"
 #include "DetailCategoryBuilder.h"
 #include "DetailLayoutBuilder.h"
 #include "DetailWidgetRow.h"
+#include "Dialogs/DlgPickPath.h"
 #include "Editor.h"
 #include "Engine/Selection.h"
+#include "IContentBrowserDataModule.h"
+#include "IContentBrowserSingleton.h"
+#include "Interfaces/IMainFrameModule.h"
+#include "IStructureDetailsView.h"
 #include "MaterialShared.h"
+#include "Misc/MessageDialog.h"
 #include "Misc/ScopedSlowTask.h"
 #include "PropertyHandle.h"
 #include "Widgets/DeclarativeSyntaxSupport.h"
 #include "Widgets/Input/SButton.h"
+#include "Widgets/Layout/SUniformGridPanel.h"
 #include "WorldPartition/WorldPartition.h"
 #include "WorldPartition/HLOD/HLODActor.h"
 
@@ -37,10 +47,21 @@ void FWorldPartitionHLODDetailsCustomization::CustomizeDetails(IDetailLayoutBuil
 		+SVerticalBox::Slot()
 		[
 			SNew(SButton)
+			.HAlign(EHorizontalAlignment::HAlign_Center)
+			.VAlign(EVerticalAlignment::VAlign_Center)
 			.Text(LOCTEXT("BuildHLODButtonText", "Build HLOD"))
 			.OnClicked(this, &FWorldPartitionHLODDetailsCustomization::OnBuildHLOD)
 			.IsEnabled(this, &FWorldPartitionHLODDetailsCustomization::CanBuildHLOD)
 		]
+		+ SVerticalBox::Slot()
+		[
+			SNew(SButton)
+			.HAlign(EHorizontalAlignment::HAlign_Center)
+			.VAlign(EVerticalAlignment::VAlign_Center)
+			.Text(LOCTEXT("ExportAssetsButtonText", "Export HLOD Assets"))
+			.OnClicked(this, &FWorldPartitionHLODDetailsCustomization::OnExportAssets)
+			.IsEnabled(this, &FWorldPartitionHLODDetailsCustomization::CanExportAssets)
+		]
 	];
 }
 
@@ -135,4 +156,273 @@ FReply FWorldPartitionHLODDetailsCustomization::OnBuildHLOD()
 	return FReply::Handled();
 }
 
+bool FWorldPartitionHLODDetailsCustomization::CanExportHLODAssets(const AWorldPartitionHLOD* HLODActor)
+{
+	if (HLODActor->IsTemplate())
+	{
+		return false;
+	}
+
+	if (HLODActor->GetSourceActors() == nullptr)
+	{
+		return false;
+	}
+
+	bool bHasPrivateStaticMesh = false;
+	ForEachObjectWithPackage(HLODActor->GetPackage(), [&bHasPrivateStaticMesh](UObject* InObject)
+	{
+		bHasPrivateStaticMesh = InObject->IsA<UStaticMesh>();
+		return !bHasPrivateStaticMesh; // Continue until we find a static mesh
+	});
+
+	return bHasPrivateStaticMesh;
+}
+
+bool FWorldPartitionHLODDetailsCustomization::CanExportAssets() const
+{
+	TArray<AWorldPartitionHLOD*> SelectedHLODActors = GetSelectedHLODActors();
+	if (SelectedHLODActors.IsEmpty())
+	{
+		return false;
+	}
+
+	return Algo::AnyOf(SelectedHLODActors, CanExportHLODAssets);
+}
+
+static EAppReturnType::Type EditExportHLODAssetsParams(FExportHLODAssetsParams& InOutParams, const TArray<AWorldPartitionHLOD*>& InHLODActorsToExport)
+{
+	// Work on a copy; commit only if user accepts.
+	FExportHLODAssetsParams WorkingCopy = InOutParams;
+
+	// Setup property editor for a struct
+	FPropertyEditorModule& PropertyEditorModule = FModuleManager::LoadModuleChecked<FPropertyEditorModule>("PropertyEditor");
+
+	FDetailsViewArgs DetailsArgs;
+	DetailsArgs.bAllowSearch = false; // no search bar
+	DetailsArgs.bShowOptions = false;
+	DetailsArgs.bHideSelectionTip = true;
+	DetailsArgs.NameAreaSettings = FDetailsViewArgs::HideNameArea;
+	DetailsArgs.DefaultsOnlyVisibility = EEditDefaultsOnlyNodeVisibility::Automatic;
+
+	FStructureDetailsViewArgs StructArgs;
+	StructArgs.bShowObjects = false;
+
+	TSharedRef<IStructureDetailsView> StructDetails =
+		PropertyEditorModule.CreateStructureDetailView(DetailsArgs, StructArgs, nullptr);
+
+	TSharedRef<FStructOnScope> StructData = MakeShared<FStructOnScope>(
+		FExportHLODAssetsParams::StaticStruct(),
+		reinterpret_cast<uint8*>(&WorkingCopy) // valid: modal window blocks until closed
+	);
+	StructDetails->SetStructureData(StructData);
+	
+	auto CanExport = [&WorkingCopy, &InHLODActorsToExport](FText& OutExportErrorMessage) -> bool
+	{
+		OutExportErrorMessage = FText::GetEmpty();
+
+		FExportHLODAssetsParams TestExportHLODAssetsParams = WorkingCopy;
+		TestExportHLODAssetsParams.bTestExportOnly = true;
+
+		for (AWorldPartitionHLOD* HLODActor : InHLODActorsToExport)
+		{
+			FString OutErrorMessage;
+			HLODActor->ExportHLODAssets(TestExportHLODAssetsParams, OutErrorMessage);
+			if (!OutErrorMessage.IsEmpty())
+			{
+				OutExportErrorMessage = FText::FromString(OutErrorMessage);
+				return false;
+			}
+		}
+
+		return true;
+	};
+
+	FText ExportErrorMessage;
+	bool bCanExportCached = CanExport(ExportErrorMessage);
+
+	StructDetails->GetOnFinishedChangingPropertiesDelegate().AddLambda(
+		[&bCanExportCached, &CanExport, &ExportErrorMessage](const FPropertyChangedEvent&)
+		{
+			bCanExportCached = CanExport(ExportErrorMessage);
+		});	
+
+	bCanExportCached = CanExport(ExportErrorMessage);
+
+	// Modal window
+	TSharedRef<SWindow> Window = SNew(SWindow)
+		.Title(LOCTEXT("ExportHLODAssetsWindowTitle", "Export HLOD Assets"))
+		.ClientSize(FVector2D(560, 400))              
+		.SizingRule(ESizingRule::UserSized)
+		.SupportsMinimize(false)
+		.SupportsMaximize(false)
+		.IsTopmostWindow(true);
+
+	EAppReturnType::Type Result = EAppReturnType::Cancel;
+
+	Window->SetContent(
+		SNew(SBorder)
+		.Padding(4.f)
+		[
+			SNew(SVerticalBox)
+
+			// Details view
+			+ SVerticalBox::Slot()
+			.FillHeight(1.f)
+			[
+				StructDetails->GetWidget().ToSharedRef()
+			]
+
+			// Error label
+			+ SVerticalBox::Slot()
+			.AutoHeight()
+			.Padding(0.f, 6.f)
+			[
+				SNew(SBorder)
+				.Padding(6.f)
+				.BorderImage(FAppStyle::GetBrush("MessageLog.ListBorder"))
+				.BorderBackgroundColor(FStyleColors::Error)
+				[
+					SNew(STextBlock)
+					.Text_Lambda([&ExportErrorMessage]() { return ExportErrorMessage; })
+					.Visibility_Lambda([&ExportErrorMessage]() { return ExportErrorMessage.IsEmpty() ? EVisibility::Collapsed : EVisibility::Visible; })
+					.ColorAndOpacity(FStyleColors::Foreground)
+					.Font(FAppStyle::GetFontStyle("NormalBold"))
+					.WrapTextAt(500.f)
+					.Justification(ETextJustify::Left)
+				]
+			]
+
+			// Buttons
+			+ SVerticalBox::Slot()
+			.AutoHeight()
+			.HAlign(HAlign_Right)
+			.Padding(0.f, 8.f, 0.f, 0.f)
+			[
+				SNew(SUniformGridPanel)
+				.SlotPadding(FMargin(4.f, 0.f))
+
+				// "Export"
+				+ SUniformGridPanel::Slot(0, 0)
+				[
+					SNew(SButton)
+					.HAlign(HAlign_Center)
+					.VAlign(VAlign_Center)
+					.IsEnabled_Lambda([&bCanExportCached]() { return bCanExportCached; })
+					.OnClicked_Lambda([&Result, Window]()
+					{
+						Result = EAppReturnType::Ok;
+						Window->RequestDestroyWindow();
+						return FReply::Handled();
+					})
+					[
+						SNew(SHorizontalBox)
+
+						// Error icon (visible only when error exists)
+						+ SHorizontalBox::Slot()
+						.AutoWidth()
+						.VAlign(VAlign_Center)
+						[
+							SNew(SImage)
+							.Visibility_Lambda([&ExportErrorMessage]()
+							{
+								return ExportErrorMessage.IsEmpty() ? EVisibility::Collapsed : EVisibility::Visible;
+							})
+							.Image(FAppStyle::Get().GetBrush("Icons.Error")) // UE editor’s error icon
+						]
+
+						+ SHorizontalBox::Slot()
+						.AutoWidth()
+						.Padding(4.f, 0.f)
+						.VAlign(VAlign_Center)
+						[
+							SNew(STextBlock)
+							.Text(LOCTEXT("Export", "Export"))
+						]
+					]
+				]
+
+				// "Cancel"
+				+ SUniformGridPanel::Slot(1, 0)
+				[
+					SNew(SButton)
+					.HAlign(HAlign_Center)
+					.VAlign(VAlign_Center)
+					.Text(LOCTEXT("Cancel", "Cancel"))
+					.OnClicked_Lambda([Window]()
+					{
+						Window->RequestDestroyWindow();
+						return FReply::Handled();
+					})
+				]
+			]
+		]
+	);
+
+	IMainFrameModule& MainFrameModule = FModuleManager::GetModuleChecked<IMainFrameModule>("MainFrame");
+	FSlateApplication::Get().AddModalWindow(Window, MainFrameModule.GetParentWindow());
+
+	if (Result == EAppReturnType::Ok)
+	{
+		InOutParams = WorkingCopy; // commit
+	}
+
+	return Result;
+}
+
+FReply FWorldPartitionHLODDetailsCustomization::OnExportAssets()
+{
+	TArray<AWorldPartitionHLOD*> HLODActors;
+	Algo::CopyIf(GetSelectedHLODActors(), HLODActors, CanExportHLODAssets);
+
+	// Static so that settings persist in the same user session
+	static FExportHLODAssetsParams ExportHLODAssetsParams;
+	if (EditExportHLODAssetsParams(ExportHLODAssetsParams, HLODActors) != EAppReturnType::Ok)
+	{
+		return FReply::Handled();
+	}
+	
+	TArray<UObject*> ExportedAssets;
+	const uint32 SlowTaskStepsCount = HLODActors.Num() + 1; // +1 for the content browser sync operation...
+
+	FScopedSlowTask SlowTask(SlowTaskStepsCount, FText::FromString(TEXT("Exporting HLOD Assets...")));
+	SlowTask.MakeDialogDelayed(1.f, true);
+
+	for (int32 Index = 0; Index < HLODActors.Num(); Index++)
+	{
+		AWorldPartitionHLOD* HLODActor = HLODActors[Index];
+		SlowTask.EnterProgressFrame(1.f, FText::FromString(FString::Printf(TEXT("Exporting HLOD Assets (%d/%d)"), Index + 1, HLODActors.Num())));
+
+		if (SlowTask.ShouldCancel())
+		{
+			break;
+		}
+
+		FString OutErrorMessage;
+		ExportedAssets += HLODActor->ExportHLODAssets(ExportHLODAssetsParams, OutErrorMessage);
+		if (!OutErrorMessage.IsEmpty())
+		{
+			FMessageDialog::Open(EAppMsgCategory::Error, EAppMsgType::Ok, FText::FromString(OutErrorMessage), FText::FromString(TEXT("Export failed!")));
+			break;
+		}			
+	}
+
+	if (ExportedAssets.Num())
+	{
+		SlowTask.EnterProgressFrame(1.f, FText::FromString(TEXT("Finalizing Export...")));
+
+		UContentBrowserDataSubsystem* ContentBrowserDataSubsystem = IContentBrowserDataModule::Get().GetSubsystem();
+		
+		TArray<FContentBrowserItem> ContentBrowserItems;
+		ContentBrowserDataSubsystem->EnumerateItemsForObjects(ExportedAssets, [&ContentBrowserItems](FContentBrowserItemData&& ItemData)
+		{
+			ContentBrowserItems.Emplace(ItemData);
+			return true;
+		});
+
+		FContentBrowserModule& ContentBrowserModule = FModuleManager::LoadModuleChecked<FContentBrowserModule>("ContentBrowser");
+		ContentBrowserModule.Get().SyncBrowserToItems(ContentBrowserItems);
+	}
+
+	return FReply::Handled();
+}
 #undef LOCTEXT_NAMESPACE
diff --git a/Engine/Source/Editor/WorldPartitionEditor/Private/WorldPartition/Customizations/WorldPartitionHLODDetailsCustomization.h b/Engine/Source/Editor/WorldPartitionEditor/Private/WorldPartition/Customizations/WorldPartitionHLODDetailsCustomization.h
--- a/Engine/Source/Editor/WorldPartitionEditor/Private/WorldPartition/Customizations/WorldPartitionHLODDetailsCustomization.h
+++ b/Engine/Source/Editor/WorldPartitionEditor/Private/WorldPartition/Customizations/WorldPartitionHLODDetailsCustomization.h
@@ -19,15 +19,19 @@ protected:
 	FWorldPartitionHLODDetailsCustomization()
 	{}
 
-	bool CanBuildHLOD() const;
-	FReply OnBuildHLOD();
-
 private:
 	/** IDetailCustomization interface */
 	virtual void CustomizeDetails(IDetailLayoutBuilder& InDetailBuilder) override;
 
 	TArray<AWorldPartitionHLOD*> GetSelectedHLODActors() const;
 
+	bool CanBuildHLOD() const;
+	FReply OnBuildHLOD();
+
+	static bool CanExportHLODAssets(const AWorldPartitionHLOD* HLODActor);
+	bool CanExportAssets() const;
+	FReply OnExportAssets();
+
 private:
 	TArray<TWeakObjectPtr<UObject>> SelectedObjects;
-};
\ No newline at end of file
+};
diff --git a/Engine/Source/Editor/WorldPartitionEditor/WorldPartitionEditor.Build.cs b/Engine/Source/Editor/WorldPartitionEditor/WorldPartitionEditor.Build.cs
--- a/Engine/Source/Editor/WorldPartitionEditor/WorldPartitionEditor.Build.cs
+++ b/Engine/Source/Editor/WorldPartitionEditor/WorldPartitionEditor.Build.cs
@@ -9,27 +9,29 @@ public class WorldPartitionEditor : ModuleRules
         PrivateDependencyModuleNames.AddRange(
             new string[] {
 				"ApplicationCore",
+				"ContentBrowser",
+				"ContentBrowserData",
 				"Core",
 				"CoreUObject",
+				"DataLayerEditor",
+				"DeveloperSettings",
 				"EditorFramework",
 				"EditorSubsystem",
 				"EditorWidgets",
 				"Engine",
 				"InputCore",
-				"LevelEditor",
+				"LevelEditor",	
+				"MainFrame",
+				"PropertyEditor",				
+				"SceneOutliner",
 				"Slate",
 				"SlateCore",
-				"UnrealEd",
-				"WorldBrowser",
-				"MainFrame",
-				"PropertyEditor",
-				"DeveloperSettings",
-				"ToolMenus",
 				"RenderCore",
 				"Renderer",
-				"RHI",
-				"SceneOutliner",
-				"DataLayerEditor"
+				"RHI",				
+				"ToolMenus",
+				"UnrealEd",
+				"WorldBrowser"
 			}
 		);
 
diff --git a/Engine/Source/Runtime/Engine/Private/WorldPartition/HLOD/HLODActor.cpp b/Engine/Source/Runtime/Engine/Private/WorldPartition/HLOD/HLODActor.cpp
--- a/Engine/Source/Runtime/Engine/Private/WorldPartition/HLOD/HLODActor.cpp
+++ b/Engine/Source/Runtime/Engine/Private/WorldPartition/HLOD/HLODActor.cpp
@@ -18,13 +18,19 @@
 #include "WorldPartition/WorldPartitionHelpers.h"
 
 #if WITH_EDITOR
+#include "AssetRegistry/AssetRegistryModule.h"
 #include "Components/StaticMeshComponent.h"
 #include "Editor.h"
+#include "Engine/StaticMeshSocket.h"
+#include "Engine/Texture.h"
 #include "HAL/FileManager.h"
-#include "Modules/ModuleManager.h"
+#include "MeshDescription.h"
 #include "Misc/ArchiveMD5.h"
+#include "Modules/ModuleManager.h"
 #include "PhysicsEngine/BodySetup.h"
+#include "Serialization/ArchiveReplaceObjectRef.h"
 #include "Subsystems/AssetEditorSubsystem.h"
+#include "Subsystems/EditorAssetSubsystem.h"
 #include "UObject/FortniteMainBranchObjectVersion.h"
 #include "WorldPartition/HLOD/HLODActorDesc.h"
 #include "WorldPartition/HLOD/HLODLayer.h"
@@ -728,4 +734,244 @@ void AWorldPartitionHLOD::UpdateHLODBuildReportContent(const FString& InHLODBuil
 	HLODBuildReport += LINE_TERMINATOR;
 }
 
+TArray<UObject*> AWorldPartitionHLOD::ExportHLODAssets(const FExportHLODAssetsParams& ExportHLODAssetsParams, FString& OutErrorMessage) const
+{
+	TArray<UObject*> ExportedAssets;
+	OutErrorMessage = TEXT("");
+		
+	auto ShouldExportAsset = [ActorPackage = GetPackage()](const UObject* InAsset)
+	{
+		// Skip null entries and assets that are not in the HLOD package (no need for export)
+		return InAsset && InAsset->IsInPackage(ActorPackage);
+	};
+
+	// Unbuilt HLODs will have no static mesh components, we can ignore them
+	UStaticMeshComponent* MeshComp = FindComponentByClass<UStaticMeshComponent>();
+	if (!MeshComp)
+	{
+		return ExportedAssets;
+	}
+
+	UStaticMesh* StaticMesh = MeshComp->GetStaticMesh();
+	if (!ShouldExportAsset(StaticMesh))
+	{
+		return ExportedAssets;
+	}
+
+	FString ExportRootPath = ExportHLODAssetsParams.ExportRootPath.Path;
+	FPaths::NormalizeDirectoryName(ExportRootPath);
+
+	if (ExportRootPath.IsEmpty() || !FPackageName::IsValidLongPackageName(ExportRootPath))
+	{
+		OutErrorMessage = FString::Printf(TEXT("Invalid export path: %s"), *ExportRootPath);
+		return ExportedAssets;
+	}	
+
+	UEditorAssetSubsystem* EditorAssetSubsystem = GEditor->GetEditorSubsystem<UEditorAssetSubsystem>();
+
+	FString MeshPackagePath = ExportRootPath / GetActorLabel();
+	FString MeshAssetName = FPackageName::GetLongPackageAssetName(MeshPackagePath);
+	FString AssetsRootPath = FPackageName::GetLongPackagePath(MeshPackagePath);
+
+	if (EditorAssetSubsystem->DoesAssetExist(MeshPackagePath))
+	{
+		OutErrorMessage = FString::Printf(TEXT("Mesh asset already exists at %s"), *MeshPackagePath);
+		return ExportedAssets;
+	}
+
+
+	for (int32 MatIndex = 0; MatIndex < StaticMesh->GetStaticMaterials().Num(); ++MatIndex)
+	{
+		UMaterialInterface* Mat = MeshComp->GetMaterial(MatIndex);
+
+		// Skip null entries and materials that are not in the HLOD package (no need for export)
+		if (!ShouldExportAsset(Mat))
+		{
+			continue;
+		}
+
+		FString MatPath = AssetsRootPath / Mat->GetName();
+		if (EditorAssetSubsystem->DoesAssetExist(MatPath))
+		{
+			OutErrorMessage = FString::Printf(TEXT("Material asset already exists at %s"), *MatPath);
+			return ExportedAssets;
+		}
+
+		TArray<UTexture*> UsedTextures;
+		Mat->GetUsedTextures(UsedTextures);
+		for (UTexture* Tex : UsedTextures)
+		{
+			// Skip null entries and textures that are not in the HLOD package (no need for export)
+			if (!ShouldExportAsset(Tex))
+			{
+				continue;
+			}
+
+			FString TexPath = AssetsRootPath / Tex->GetName();
+			if (EditorAssetSubsystem->DoesAssetExist(TexPath))
+			{
+				OutErrorMessage = FString::Printf(TEXT("Texture asset already exists at %s"), *TexPath);
+				return ExportedAssets;
+			}
+		}
+	}
+
+	if (ExportHLODAssetsParams.bTestExportOnly)
+	{
+		OutErrorMessage = TEXT("");
+		return ExportedAssets;
+	}
+
+	// Duplicate mesh
+	UPackage* MeshPackage = CreatePackage(*MeshPackagePath);
+	UStaticMesh* NewMesh = DuplicateObject<UStaticMesh>(StaticMesh, MeshPackage, *MeshAssetName);
+	if (!NewMesh)
+	{
+		OutErrorMessage = TEXT("Failed to duplicate static mesh.");
+		return ExportedAssets;
+	}
+
+	// Offset calculation	
+	if (ExportHLODAssetsParams.MeshOrigin == EExportHLODMeshOrigin::Actor)
+	{
+		FVector Offset = -GetActorLocation();
+
+		// Offset vertex positions
+		if (NewMesh->IsMeshDescriptionValid(0))
+		{
+			FMeshDescription* MeshDesc = NewMesh->GetMeshDescription(0);
+			if (MeshDesc)
+			{
+				TVertexAttributesRef<FVector3f> VertexPositions = MeshDesc->GetVertexPositions();
+				for (const FVertexID VertexID : MeshDesc->Vertices().GetElementIDs())
+				{
+					VertexPositions[VertexID] += (FVector3f)Offset;
+				}
+				NewMesh->CommitMeshDescription(0);
+			}
+		}
+
+		// Offset sockets
+		for (UStaticMeshSocket* Socket : NewMesh->Sockets)
+		{
+			if (Socket)
+			{
+				Socket->RelativeLocation += Offset;
+			}
+		}
+
+		// Offset collision data
+		if (NewMesh->GetBodySetup())
+		{
+			UBodySetup* Body = NewMesh->GetBodySetup();
+
+			for (FKConvexElem& Convex : Body->AggGeom.ConvexElems)
+			{
+				for (FVector& V : Convex.VertexData)
+				{
+					V += Offset;
+				}
+				Convex.UpdateElemBox();
+			}
+			for (FKBoxElem& Elem : Body->AggGeom.BoxElems)
+			{
+				Elem.Center += Offset;
+			}
+
+			for (FKSphereElem& Elem : Body->AggGeom.SphereElems)
+			{
+				Elem.Center += Offset;
+			}
+
+			for (FKSphylElem& Elem : Body->AggGeom.SphylElems)
+			{
+				Elem.Center += Offset;
+			}
+
+			Body->InvalidatePhysicsData();
+			Body->CreatePhysicsMeshes();
+		}
+	}
+
+	// Helper to finalize asset export
+	auto OnAssetExported = [&ExportedAssets](UObject* InAsset)
+	{
+		InAsset->SetFlags(RF_Public | RF_Standalone);
+		FAssetRegistryModule::AssetCreated(InAsset);
+		InAsset->MarkPackageDirty();
+		InAsset->PostEditChange();
+		ExportedAssets.Add(InAsset);
+	};
+
+	// Duplicate materials and textures
+	TMap<UMaterialInterface*, UMaterialInterface*> DuplicatedMaterials;
+	TMap<UTexture*, UTexture*> DuplicatedTextures;
+	for (int32 MatIndex = 0; MatIndex < StaticMesh->GetStaticMaterials().Num(); ++MatIndex)
+	{
+		UMaterialInterface* Mat = MeshComp->GetMaterial(MatIndex);
+
+		// Skip null entries and materials that are not in the HLOD package (no need for export)
+		if (!ShouldExportAsset(Mat))
+		{
+			continue;
+		}
+
+		// Skip materials that were already exported
+		if (DuplicatedMaterials.Contains(Mat))
+		{
+			continue;
+		}
+
+		FString MatPath = AssetsRootPath / Mat->GetName();
+		UPackage* MatPackage = CreatePackage(*MatPath);
+		UMaterialInterface* NewMat = DuplicateObject<UMaterialInterface>(Mat, MatPackage, *Mat->GetName());
+
+		if (NewMat)
+		{
+			DuplicatedMaterials.Emplace(Mat, NewMat);
+
+			TArray<UTexture*> UsedTextures;
+			NewMat->GetUsedTextures(UsedTextures);
+
+			for (UTexture* Tex : UsedTextures)
+			{
+				// Skip null entries and textures that are not in the HLOD package (no need for export)
+				if (!ShouldExportAsset(Tex))
+				{
+					continue;
+				}
+
+				// Skip textures that were already exported
+				if (DuplicatedTextures.Contains(Tex))
+				{
+					continue;
+				}
+
+				FString TexPath = AssetsRootPath / Tex->GetName();
+				UPackage* TexPackage = CreatePackage(*TexPath);
+				UTexture* NewTex = DuplicateObject<UTexture>(Tex, TexPackage, *Tex->GetName());
+				if (NewTex)
+				{
+					DuplicatedTextures.Emplace(Tex, NewTex);
+					OnAssetExported(NewTex);
+				}
+			}
+
+			// Replace all the original textures by their duplicate in the new material
+			FArchiveReplaceObjectRef<UTexture> MeshReplacer(NewMat, DuplicatedTextures);
+
+			OnAssetExported(NewMat);
+		}
+	}
+
+	// Replace all the original materials by their duplicate in the new mesh
+	FArchiveReplaceObjectRef<UMaterialInterface> MeshReplacer(NewMesh, DuplicatedMaterials);
+
+	OnAssetExported(NewMesh);
+
+	OutErrorMessage = TEXT("");
+
+	return ExportedAssets;
+}
+
 #endif // WITH_EDITOR
\ No newline at end of file
diff --git a/Engine/Source/Runtime/Engine/Public/WorldPartition/HLOD/HLODActor.h b/Engine/Source/Runtime/Engine/Public/WorldPartition/HLOD/HLODActor.h
--- a/Engine/Source/Runtime/Engine/Public/WorldPartition/HLOD/HLODActor.h
+++ b/Engine/Source/Runtime/Engine/Public/WorldPartition/HLOD/HLODActor.h
@@ -39,6 +39,29 @@ struct FWorldPartitionHLODBuildEvent
 DECLARE_MULTICAST_DELEGATE_OneParam(FWorldPartitionHLODBuildEventDelegate, const FWorldPartitionHLODBuildEvent&)
 #endif
 
+UENUM(BlueprintType)
+enum class EExportHLODMeshOrigin : uint8
+{
+	Actor   UMETA(DisplayName = "Actor Origin"),		// Origin is at actor location
+	World   UMETA(DisplayName = "World Origin"),        // Origin is at world <0,0,0>
+};
+
+USTRUCT(BlueprintType)
+struct FExportHLODAssetsParams
+{
+	GENERATED_USTRUCT_BODY()
+
+#if WITH_EDITORONLY_DATA
+	UPROPERTY(EditAnywhere, Category = "Export Options")
+	EExportHLODMeshOrigin MeshOrigin = EExportHLODMeshOrigin::World;
+
+	UPROPERTY(EditAnywhere, Category = "Export Options", meta=(ContentDir))
+	FDirectoryPath ExportRootPath;
+
+	bool bTestExportOnly = false;
+#endif
+};
+
 UCLASS(NotPlaceable, MinimalAPI, HideCategories=(Rendering, Replication, Collision, Physics, Navigation, Networking, Input, Actor, LevelInstance, Cooking))
 class AWorldPartitionHLOD : public AActor, public IWorldPartitionHLODObject
 {
@@ -102,6 +125,9 @@ public:
 
 	ENGINE_API static FWorldPartitionHLODBuildEventDelegate& GetHLODBuildEventDelegate();
 
+	UFUNCTION(BlueprintCallable, Category = "World Partition", meta = (DevelopmentOnly))
+	ENGINE_API TArray<UObject*> ExportHLODAssets(const struct FExportHLODAssetsParams& ExportHLODAssetsParams, FString& OutErrorMessage) const;
+
 	ENGINE_API void UpdateHLODBuildReportHeader();
 	ENGINE_API void UpdateHLODBuildReportContent(const FString& InHLODBuildReportContent);
 
@@ -148,6 +174,7 @@ protected:
 	virtual bool ShouldImport(FStringView ActorPropString, bool IsMovingLevel) override { return false; }
 	virtual bool IsLockLocation() const override { return true; }
 	virtual bool IsUserManaged() const override { return false; }
+	virtual bool IsActorLabelEditable() const override { return false; }
 	ENGINE_API virtual void GetAssetRegistryTags(FAssetRegistryTagsContext Context) const override;
 #endif
 	//~ End AActor Interface.
