[HLOD] Consider whether child HLOD Actors have all components registered, instead of only checking whether they're loaded, when determining HLOD Actor visibility in the Editor.
Fixes cases where child HLOD Actors are loaded but removed from the world, yet still prevent the parent HLOD actor from being visible, resulting in no HLOD being shown for that area.

diff --git a/Engine/Source/Editor/WorldPartitionEditor/Private/WorldPartition/HLOD/HLODEditorData.cpp b/Engine/Source/Editor/WorldPartitionEditor/Private/WorldPartition/HLOD/HLODEditorData.cpp
--- a/Engine/Source/Editor/WorldPartitionEditor/Private/WorldPartition/HLOD/HLODEditorData.cpp
+++ b/Engine/Source/Editor/WorldPartitionEditor/Private/WorldPartition/HLOD/HLODEditorData.cpp
@@ -325,9 +325,20 @@ void FHLODSceneNode::UpdateVisibility(const FVector& InCameraLocation, double In
 		// Perform distance culling
 		if (bNodeShouldBeVisible)
 		{
-			// Do not perform near culling if this HLOD has no child HLOD or none of them is loaded
-			const bool bHasLoadedChildrenHLODs = Algo::AnyOf(ChildrenHLODs, [](FHLODSceneNode* ChildHLODNode) { return ChildHLODNode->HLODActorHandle.IsLoaded(); });
-			if (bHasLoadedChildrenHLODs)
+			// Do not perform near culling if this HLOD has no child HLOD or none of them is loaded and has all components registered
+			const bool bHasLoadedAndRegisteredChildrenHLODs = Algo::AnyOf(ChildrenHLODs, [](FHLODSceneNode* ChildHLODNode)
+			{
+				if(AActor* HLODActor = ChildHLODNode->HLODActorHandle.GetActor())
+				{
+					if (HLODActor->HasActorRegisteredAllComponents())
+					{
+						return true;
+					}
+				}
+				return false;
+			});
+
+			if (bHasLoadedAndRegisteredChildrenHLODs)
 			{
 				const double DistanceSquared = Bounds.ComputeSquaredDistanceFromBoxToPoint(InCameraLocation);
 				const bool bNearCulled = DistanceSquared < FMath::Square(HLODActor->GetMinVisibleDistance());
