[HLOD] Add 'Build HLODs From Selected Region(s)' action to the WP minimap context menu.
Create a new HLOD section containing this action and 'Build HLODs For Selection'.

diff --git a/Engine/Source/Editor/WorldPartitionEditor/Private/WorldPartition/SWorldPartitionEditorGrid2D.cpp b/Engine/Source/Editor/WorldPartitionEditor/Private/WorldPartition/SWorldPartitionEditorGrid2D.cpp
--- a/Engine/Source/Editor/WorldPartitionEditor/Private/WorldPartition/SWorldPartitionEditorGrid2D.cpp
+++ b/Engine/Source/Editor/WorldPartitionEditor/Private/WorldPartition/SWorldPartitionEditorGrid2D.cpp
@@ -231,7 +231,8 @@ void SWorldPartitionEditorGrid2D::FEditorCommands::RegisterCommands()
 {
 	// Context Menu
 	UI_COMMAND(CreateRegionFromSelection, "Load Region From Selection", "Load region from selection.", EUserInterfaceActionType::Button, FInputChord());
-	UI_COMMAND(BuildHLODsForSelection, "Build HLODs for selection", "Build HLODs in selected area.", EUserInterfaceActionType::Button, FInputChord());
+	UI_COMMAND(BuildHLODsForSelection, "Build HLODs For Selection", "Build HLODs in selected area.", EUserInterfaceActionType::Button, FInputChord());
+	UI_COMMAND(BuildHLODsForSelectedRegions, "Build HLODs For Selected Region(s)", "Build HLODs for the selected region(s).", EUserInterfaceActionType::Button, FInputChord());
 	UI_COMMAND(ConvertSelectedRegionsToActors, "Convert Selected Region(s) To Actor(s)", "Convert the selected region(s) to actor(s).", EUserInterfaceActionType::Button, FInputChord());
 	UI_COMMAND(LoadSelectedRegions, "Load Selected Region(s)", "Load the selected region(s).", EUserInterfaceActionType::Button, FInputChord());
 	UI_COMMAND(UnloadSelectedRegions, "Unload Selected Region(s)", "Unload the selected region(s).", EUserInterfaceActionType::Button, FInputChord());
@@ -574,6 +575,21 @@ void SWorldPartitionEditorGrid2D::BindCommands()
 		return false;
 	};
 
+	auto CanBuildHLODsForSelectedRegions = [this]()
+	{
+		for (const TWeakInterfacePtr<IWorldPartitionActorLoaderInterface>& SelectedLoaderAdapter : SelectedLoaderInterfaces)
+		{
+			if (IWorldPartitionActorLoaderInterface* LoaderInterface = SelectedLoaderAdapter.Get())
+			{
+				if (LoaderInterface->GetLoaderAdapter()->GetBoundingBox().IsSet())
+				{
+					return true;
+				}
+			}
+		}
+		return false;
+	};
+
 	auto CanConvertSelectedRegionsToActors = [this]()
 	{
 		for (const TWeakInterfacePtr<IWorldPartitionActorLoaderInterface>& SelectedLoaderAdapter : SelectedLoaderInterfaces)
@@ -609,6 +625,7 @@ void SWorldPartitionEditorGrid2D::BindCommands()
 
 	CommandList->MapAction(Commands.CreateRegionFromSelection, FExecuteAction::CreateSP(this, &SWorldPartitionEditorGrid2D::CreateRegionFromSelection), FCanExecuteAction::CreateLambda(CanCreateRegionFromSelection));
 	CommandList->MapAction(Commands.BuildHLODsForSelection, FExecuteAction::CreateSP(this, &SWorldPartitionEditorGrid2D::BuildHLODsForSelection), FCanExecuteAction::CreateLambda(CanCreateRegionFromSelection));
+	CommandList->MapAction(Commands.BuildHLODsForSelectedRegions, FExecuteAction::CreateSP(this, &SWorldPartitionEditorGrid2D::BuildHLODsForSelectedRegions), FCanExecuteAction::CreateLambda(CanBuildHLODsForSelectedRegions));
 	CommandList->MapAction(Commands.LoadSelectedRegions, FExecuteAction::CreateSP(this, &SWorldPartitionEditorGrid2D::LoadSelectedRegions), FCanExecuteAction::CreateLambda(CanLoadSelectedRegions));
 	CommandList->MapAction(Commands.UnloadSelectedRegions, FExecuteAction::CreateSP(this, &SWorldPartitionEditorGrid2D::UnloadSelectedRegions), FCanExecuteAction::CreateLambda(CanUnloadSelectedRegions));
 	CommandList->MapAction(Commands.ConvertSelectedRegionsToActors, FExecuteAction::CreateSP(this, &SWorldPartitionEditorGrid2D::ConvertSelectedRegionsToActors), FCanExecuteAction::CreateLambda(CanConvertSelectedRegionsToActors));
@@ -705,6 +722,27 @@ void SWorldPartitionEditorGrid2D::BuildHLODsForSelection()
 	Refresh();
 }
 
+void SWorldPartitionEditorGrid2D::BuildHLODsForSelectedRegions()
+{
+	TArray<FBox> FilterVolumes;
+	for (const TWeakInterfacePtr<IWorldPartitionActorLoaderInterface>& SelectedLoaderAdapter : SelectedLoaderInterfaces)
+	{
+		if (IWorldPartitionActorLoaderInterface* LoaderInterface = SelectedLoaderAdapter.Get())
+		{
+			TOptional<FBox> BoundingBox = LoaderInterface->GetLoaderAdapter()->GetBoundingBox();
+			if (BoundingBox.IsSet())
+			{
+				FilterVolumes.Add(BoundingBox.GetValue());
+			}
+		}
+	}
+
+	UWorldPartitionHLODsBuilder::BuildWithFilters(GetWorld(), { MakeShared<FHLODCreationVolumeFilter>(MoveTemp(FilterVolumes)) });
+
+	ClearSelection();
+	Refresh();
+}
+
 void SWorldPartitionEditorGrid2D::LoadSelectedRegions()
 {
 	for (const TWeakInterfacePtr<IWorldPartitionActorLoaderInterface>& SelectedLoaderAdapter : SelectedLoaderInterfaces)
@@ -934,25 +972,30 @@ TSharedRef<SWidget> SWorldPartitionEditorGrid2D::GenerateContextualMenu() const
 	static const FName MenuName(TEXT("WorldPartition.ContextualMenu"));
 	
 	const FEditorCommands& Commands = FEditorCommands::Get();
-	UToolMenu* ConxtextualMenu = UToolMenus::Get()->RegisterMenu(MenuName);
+	UToolMenu* ContextualMenu = UToolMenus::Get()->RegisterMenu(MenuName);
 
 	if (GetDefault<UWorldPartitionEditorSettings>()->GetEnableLoadingInEditor())
 	{
 		static const FName SectionSelectionName(TEXT("ContextMenu.Selection"));
-		FToolMenuSection& SectionSelection = ConxtextualMenu->AddSection(SectionSelectionName, LOCTEXT("WorldPartitionSelectionHeader", "Selection"));
+		FToolMenuSection& SectionSelection = ContextualMenu->AddSection(SectionSelectionName, LOCTEXT("WorldPartitionSelectionHeader", "Selection"));
 		SectionSelection.AddMenuEntry(Commands.CreateRegionFromSelection);
-		SectionSelection.AddMenuEntry(Commands.BuildHLODsForSelection);
 		
 		static const FName SectionRegionsName(TEXT("ContextMenu.Regions"));
-		FToolMenuSection& SectionRegions = ConxtextualMenu->AddSection(SectionRegionsName, LOCTEXT("WorldPartitionRegionsHeader", "Region(s)"));
+		FToolMenuSection& SectionRegions = ContextualMenu->AddSection(SectionRegionsName, LOCTEXT("WorldPartitionRegionsHeader", "Region(s)"));
 		SectionRegions.AddMenuEntry(Commands.LoadSelectedRegions);
 		SectionRegions.AddMenuEntry(Commands.UnloadSelectedRegions);
-		SectionRegions.AddSeparator(NAME_None);
+		SectionRegions.AddSeparator(TEXT("RegionsSeparator"));
 		SectionRegions.AddMenuEntry(Commands.ConvertSelectedRegionsToActors);
+
+		static const FName SectionHLODName(TEXT("ContextMenu.HLOD"));
+		FToolMenuSection& SectionHLOD = ContextualMenu->AddSection(SectionHLODName, LOCTEXT("WorldPartitionHLODHeader", "HLOD"));
+		SectionHLOD.AddMenuEntry(Commands.BuildHLODsForSelection);
+		SectionHLOD.AddMenuEntry(Commands.BuildHLODsForSelectedRegions);
+
 	}
 
 	static const FName SectionMiscName(TEXT("ContextMenu.Misc"));
-	FToolMenuSection& SectionMisc = ConxtextualMenu->AddSection(SectionMiscName, LOCTEXT("WorldPartitionMiscHeader", "Misc"));
+	FToolMenuSection& SectionMisc = ContextualMenu->AddSection(SectionMiscName, LOCTEXT("WorldPartitionMiscHeader", "Misc"));
 	SectionMisc.AddMenuEntry(Commands.MoveCameraHere);
 
 	if (!GetDefault<UWorldPartitionEditorSettings>()->GetDisableBugIt())
diff --git a/Engine/Source/Editor/WorldPartitionEditor/Public/WorldPartition/SWorldPartitionEditorGrid2D.h b/Engine/Source/Editor/WorldPartitionEditor/Public/WorldPartition/SWorldPartitionEditorGrid2D.h
--- a/Engine/Source/Editor/WorldPartitionEditor/Public/WorldPartition/SWorldPartitionEditorGrid2D.h
+++ b/Engine/Source/Editor/WorldPartitionEditor/Public/WorldPartition/SWorldPartitionEditorGrid2D.h
@@ -27,6 +27,7 @@ protected:
 		// Context Menu
 		TSharedPtr<FUICommandInfo> CreateRegionFromSelection;
 		TSharedPtr<FUICommandInfo> BuildHLODsForSelection;
+		TSharedPtr<FUICommandInfo> BuildHLODsForSelectedRegions;
 		TSharedPtr<FUICommandInfo> LoadSelectedRegions;
 		TSharedPtr<FUICommandInfo> UnloadSelectedRegions;
 		TSharedPtr<FUICommandInfo> UnloadHoveredRegion;
@@ -95,6 +96,7 @@ protected:
 
 	void CreateRegionFromSelection();
 	void BuildHLODsForSelection();
+	void BuildHLODsForSelectedRegions();
 	void LoadSelectedRegions();
 	void UnloadSelectedRegions();
 	void ConvertSelectedRegionsToActors();
