[HLOD] Force HLOD visibility refresh after the HLOD build with filters completes to fix an issue where higher-level HLODs could remain visible when they shouldn't.
Add a new delegate to allow systems to be notified when HLOD build with filters completes.

#Synced-CL 47987596

diff --git a/Engine/Source/Editor/UnrealEd/Private/WorldPartition/WorldPartitionHLODsBuilder.cpp b/Engine/Source/Editor/UnrealEd/Private/WorldPartition/WorldPartitionHLODsBuilder.cpp
--- a/Engine/Source/Editor/UnrealEd/Private/WorldPartition/WorldPartitionHLODsBuilder.cpp
+++ b/Engine/Source/Editor/UnrealEd/Private/WorldPartition/WorldPartitionHLODsBuilder.cpp
@@ -145,6 +145,8 @@ namespace BuildEvaluationStats
 	}
 }
 
+UWorldPartitionHLODsBuilder::FOnWorldPartitionHLODBuildWithFiltersCompleted UWorldPartitionHLODsBuilder::OnWorldPartitionHLODBuildWithFiltersCompleted;
+
 UWorldPartitionHLODsBuilder::UWorldPartitionHLODsBuilder(const FObjectInitializer& ObjectInitializer)
 	: Super(ObjectInitializer)
 	, BuilderIdx(INDEX_NONE)
@@ -240,6 +242,8 @@ bool UWorldPartitionHLODsBuilder::BuildWithFilters(UWorld* InWorld, const TArray
 
 	GEditor->RedrawLevelEditingViewports();
 
+	GetOnWorldPartitionHLODBuildWithFiltersCompleted().Broadcast();
+
 	return bResult;
 }
 
diff --git a/Engine/Source/Editor/UnrealEd/Public/WorldPartition/WorldPartitionHLODsBuilder.h b/Engine/Source/Editor/UnrealEd/Public/WorldPartition/WorldPartitionHLODsBuilder.h
--- a/Engine/Source/Editor/UnrealEd/Public/WorldPartition/WorldPartitionHLODsBuilder.h
+++ b/Engine/Source/Editor/UnrealEd/Public/WorldPartition/WorldPartitionHLODsBuilder.h
@@ -32,6 +32,9 @@ public:
 
 	UNREALED_API static bool BuildWithFilters(UWorld* InWorld, const TArray<TSharedPtr<IHLODCreationFilter>>& InFilters);
 
+	DECLARE_MULTICAST_DELEGATE(FOnWorldPartitionHLODBuildWithFiltersCompleted);
+	static FOnWorldPartitionHLODBuildWithFiltersCompleted& GetOnWorldPartitionHLODBuildWithFiltersCompleted() { return OnWorldPartitionHLODBuildWithFiltersCompleted; }
+
 	// UWorldPartitionBuilder interface begin
 	UNREALED_API virtual bool RequiresCommandletRendering() const override;
 	virtual ELoadingMode GetLoadingMode() const override { return ELoadingMode::Custom; }
@@ -112,4 +115,6 @@ private:
 	bool bSaveActors;
 	bool bConsiderUnsavedHLODActors;
 	TArray<TSharedPtr<IHLODCreationFilter>> Filters;
+
+	static UNREALED_API FOnWorldPartitionHLODBuildWithFiltersCompleted OnWorldPartitionHLODBuildWithFiltersCompleted;
 };
\ No newline at end of file
diff --git a/Engine/Source/Editor/WorldPartitionEditor/Private/WorldPartition/HLOD/HLODEditorSubsystem.cpp b/Engine/Source/Editor/WorldPartitionEditor/Private/WorldPartition/HLOD/HLODEditorSubsystem.cpp
--- a/Engine/Source/Editor/WorldPartitionEditor/Private/WorldPartition/HLOD/HLODEditorSubsystem.cpp
+++ b/Engine/Source/Editor/WorldPartitionEditor/Private/WorldPartition/HLOD/HLODEditorSubsystem.cpp
@@ -23,6 +23,7 @@
 #include "WorldPartition/WorldPartition.h"
 #include "WorldPartition/WorldPartitionActorDescInstance.h"
 #include "WorldPartition/WorldPartitionEditorSettings.h"
+#include "WorldPartition/WorldPartitionHLODsBuilder.h"
 #include "WorldPartition/WorldPartitionSubsystem.h"
 #include "WorldPartition/DataLayer/WorldDataLayers.h"
 #include "WorldPartition/DataLayer/DataLayerManager.h"
@@ -151,6 +152,7 @@ void UWorldPartitionHLODEditorSubsystem::Initialize(FSubsystemCollectionBase& Co
 	Super::Initialize(Collection);
 
 	bForceHLODStateUpdate = true;
+	bForceHLODVisibilityUpdate = true;
 	CachedCameraLocation = FVector::Zero();
 	CachedHLODMinDrawDistance = 0;
 	CachedHLODMaxDrawDistance = 0;
@@ -163,6 +165,8 @@ void UWorldPartitionHLODEditorSubsystem::Initialize(FSubsystemCollectionBase& Co
 
 	UWorldPartitionEditorSettings::OnSettingsChanged().AddUObject(this, &UWorldPartitionHLODEditorSubsystem::OnWorldPartitionEditorSettingsChanged);
 
+	UWorldPartitionHLODsBuilder::GetOnWorldPartitionHLODBuildWithFiltersCompleted().AddUObject(this, &UWorldPartitionHLODEditorSubsystem::ForceHLODVisibilityUpdate);
+
 	ApplyHLODSettingsFiltering();
 }
 
@@ -170,6 +174,8 @@ void UWorldPartitionHLODEditorSubsystem::Deinitialize()
 {
 	Super::Deinitialize();
 
+	UWorldPartitionHLODsBuilder::GetOnWorldPartitionHLODBuildWithFiltersCompleted().RemoveAll(this);
+
 	UWorldPartitionEditorSettings::OnSettingsChanged().RemoveAll(this);
 
 	GEngine->OnLevelActorListChanged().RemoveAll(this);
@@ -240,6 +246,14 @@ void UWorldPartitionHLODEditorSubsystem::OnLoaderAdapterStateChanged(const IWorl
 	ForceHLODStateUpdate();
 }
 
+void UWorldPartitionHLODEditorSubsystem::ForceHLODVisibilityUpdate()
+{
+	if (IsHLODInEditorEnabled())
+	{
+		bForceHLODVisibilityUpdate = true;
+	}
+}
+
 void UWorldPartitionHLODEditorSubsystem::ForceHLODStateUpdate()
 {
 	if (IsHLODInEditorEnabled())
@@ -255,7 +269,6 @@ void UWorldPartitionHLODEditorSubsystem::Tick(float DeltaTime)
 	UUnrealEditorSubsystem* UnrealEditorSubsystem = GEditor->GetEditorSubsystem<UUnrealEditorSubsystem>();
 
 	bool bCameraMoved = false;
-	bool bForceHLODVisibilityUpdate = false;
 	bool bClearLoadedActorState = false;
 
 	// Check cached global settings
@@ -327,6 +340,7 @@ void UWorldPartitionHLODEditorSubsystem::Tick(float DeltaTime)
 		}
 	}
 	bForceHLODStateUpdate = false;
+	bForceHLODVisibilityUpdate = false;
 }
 
 TStatId UWorldPartitionHLODEditorSubsystem::GetStatId() const
diff --git a/Engine/Source/Editor/WorldPartitionEditor/Public/WorldPartition/HLOD/HLODEditorSubsystem.h b/Engine/Source/Editor/WorldPartitionEditor/Public/WorldPartition/HLOD/HLODEditorSubsystem.h
--- a/Engine/Source/Editor/WorldPartitionEditor/Public/WorldPartition/HLOD/HLODEditorSubsystem.h
+++ b/Engine/Source/Editor/WorldPartitionEditor/Public/WorldPartition/HLOD/HLODEditorSubsystem.h
@@ -61,6 +61,8 @@ public:
 	UE_API virtual bool WriteHLODStats(const IWorldPartitionEditorModule::FWriteHLODStatsParams& Params) const;
 
 	static UE_API void AddHLODSettingsFilter(EHLODSettingsVisibility InSettingsVisibility, TSoftObjectPtr<UStruct> InStruct, FName InPropertyName);
+
+	UE_API void ForceHLODVisibilityUpdate();
 	
 private:
 	UE_API bool IsHLODInEditorEnabled();
@@ -87,6 +89,7 @@ private:
 	double CachedHLODMaxDrawDistance;
 	bool bCachedShowHLODsOverLoadedRegions;
 	bool bForceHLODStateUpdate;
+	bool bForceHLODVisibilityUpdate;
 
 	TMap<TObjectKey<UWorldPartition>, TPimplPtr<FWorldPartitionHLODEditorData>> WorldPartitionsHLODEditorData;
 
