[HLOD] Adjust the logic controlling when Build HLOD for Selection and Build HLOD for Selected Region actions are available.
Build HLOD for Selection requires a World Partition with streaming enabled.
Build HLOD for Selected Region requires a World Partition with both streaming and Loading in Editor enabled.

#Synced-CL 47987596

diff --git a/Engine/Source/Editor/WorldPartitionEditor/Private/WorldPartition/SWorldPartitionEditorGrid2D.cpp b/Engine/Source/Editor/WorldPartitionEditor/Private/WorldPartition/SWorldPartitionEditorGrid2D.cpp
--- a/Engine/Source/Editor/WorldPartitionEditor/Private/WorldPartition/SWorldPartitionEditorGrid2D.cpp
+++ b/Engine/Source/Editor/WorldPartitionEditor/Private/WorldPartition/SWorldPartitionEditorGrid2D.cpp
@@ -986,12 +986,17 @@ TSharedRef<SWidget> SWorldPartitionEditorGrid2D::GenerateContextualMenu() const
 		SectionRegions.AddMenuEntry(Commands.UnloadSelectedRegions);
 		SectionRegions.AddSeparator(TEXT("RegionsSeparator"));
 		SectionRegions.AddMenuEntry(Commands.ConvertSelectedRegionsToActors);
+	}
 
+	if (GetWorldPartition() && GetWorldPartition()->IsStreamingEnabled())
+	{
 		static const FName SectionHLODName(TEXT("ContextMenu.HLOD"));
 		FToolMenuSection& SectionHLOD = ContextualMenu->AddSection(SectionHLODName, LOCTEXT("WorldPartitionHLODHeader", "HLOD"));
 		SectionHLOD.AddMenuEntry(Commands.BuildHLODsForSelection);
-		SectionHLOD.AddMenuEntry(Commands.BuildHLODsForSelectedRegions);
-
+		if (GetDefault<UWorldPartitionEditorSettings>()->GetEnableLoadingInEditor())
+		{
+			SectionHLOD.AddMenuEntry(Commands.BuildHLODsForSelectedRegions);
+		}
 	}
 
 	static const FName SectionMiscName(TEXT("ContextMenu.Misc"));
diff --git a/Engine/Source/Editor/WorldPartitionEditor/Private/WorldPartitionEditorModule.cpp b/Engine/Source/Editor/WorldPartitionEditor/Private/WorldPartitionEditorModule.cpp
--- a/Engine/Source/Editor/WorldPartitionEditor/Private/WorldPartitionEditorModule.cpp
+++ b/Engine/Source/Editor/WorldPartitionEditor/Private/WorldPartitionEditorModule.cpp
@@ -199,7 +199,12 @@ static void CreateLevelViewportContextMenuEntries(FMenuBuilder& MenuBuilder, TAr
 			}),
 			FCanExecuteAction::CreateLambda([World]()
 			{
-				return (World.IsValid());
+				if (World.IsValid())
+				{
+					TWeakObjectPtr<UWorldPartition> WorldPartition = World->GetWorldPartition();
+					return WorldPartition.IsValid() && WorldPartition->IsStreamingEnabled();
+				}
+				return false;
 			})
 		);
 
