World Partitions: Make "Move Camera Here" (and other related minimap right-click items) ignore invisible volumes

#tests Tested in map where large BlockingVolumes are present in the sky
#rnx

#Synced-CL 48744159

diff --git a/Engine/Source/Editor/WorldPartitionEditor/Private/WorldPartition/SWorldPartitionEditorGrid2D.cpp b/Engine/Source/Editor/WorldPartitionEditor/Private/WorldPartition/SWorldPartitionEditorGrid2D.cpp
--- a/Engine/Source/Editor/WorldPartitionEditor/Private/WorldPartition/SWorldPartitionEditorGrid2D.cpp
+++ b/Engine/Source/Editor/WorldPartitionEditor/Private/WorldPartition/SWorldPartitionEditorGrid2D.cpp
@@ -174,7 +174,7 @@ static bool HitTestZFromLocation(UWorld* World, const FVector2D& WorldLocation,
 	const FVector TraceStart(WorldLocation.X, WorldLocation.Y, HALF_WORLD_MAX);
 	const FVector TraceEnd(WorldLocation.X, WorldLocation.Y, -HALF_WORLD_MAX);
 	const FCollisionQueryParams TraceParams(SCENE_QUERY_STAT(SWorldPartitionEditorGrid2D_HitTestZFromLocation), true);
-	return World->LineTraceSingleByChannel(OutResult, TraceStart, TraceEnd, ECC_Camera, TraceParams);
+	return World->LineTraceSingleByChannel(OutResult, TraceStart, TraceEnd, ECC_Visibility, TraceParams);
 }
 
 const TSet<TObjectPtr<UWorldPartitionEditorLoaderAdapter>>& GetRegisteredEditorLoaderAdapters(UWorldPartition* WorldPartition)
