[HLOD] Fix editor crash when unchecking 'Spatially Loaded' on HLOD Layer
Property handles can have multiple OnPropertyValueChanged delegates. If one delegate calls ForceRefreshDetails, the DetailLayoutBuilder can be destroyed while still in use by other delegates. Guard against this by using a WeakPtr.

#Synced-CL 49764205

diff --git a/Engine/Source/Editor/WorldPartitionEditor/Private/WorldPartition/Customizations/WorldPartitionHLODLayerDetailsCustomization.cpp b/Engine/Source/Editor/WorldPartitionEditor/Private/WorldPartition/Customizations/WorldPartitionHLODLayerDetailsCustomization.cpp
--- a/Engine/Source/Editor/WorldPartitionEditor/Private/WorldPartition/Customizations/WorldPartitionHLODLayerDetailsCustomization.cpp
+++ b/Engine/Source/Editor/WorldPartitionEditor/Private/WorldPartition/Customizations/WorldPartitionHLODLayerDetailsCustomization.cpp
@@ -16,21 +16,19 @@ TSharedRef<IDetailCustomization> FWorldPartitionHLODLayerDetailsCustomization::M
 
 void FWorldPartitionHLODLayerDetailsCustomization::CustomizeDetails(IDetailLayoutBuilder& InDetailLayoutBuilder)
 {
-	DetailLayoutBuilder = &InDetailLayoutBuilder;
-
 	UWorldPartition::WorldPartitionChangedEvent.RemoveAll(this);
 	UWorldPartition::WorldPartitionChangedEvent.AddSP(this, &FWorldPartitionHLODLayerDetailsCustomization::OnWorldPartitionChanged);
 
-	UHLODLayer* HLODLayer = !DetailLayoutBuilder->GetSelectedObjects().IsEmpty() ? Cast<UHLODLayer>(DetailLayoutBuilder->GetSelectedObjects()[0].Get()) : nullptr;
+	UHLODLayer* HLODLayer = !InDetailLayoutBuilder.GetSelectedObjects().IsEmpty() ? Cast<UHLODLayer>(InDetailLayoutBuilder.GetSelectedObjects()[0].Get()) : nullptr;
 	if (!HLODLayer)
 	{
 		return;
 	}
 		
-	TSharedPtr<IPropertyHandle> IsSpatiallyLoadedPropertyHandle = DetailLayoutBuilder->GetProperty(GET_MEMBER_NAME_CHECKED(UHLODLayer, bIsSpatiallyLoaded), UHLODLayer::StaticClass());
-	TSharedPtr<IPropertyHandle> CellSizePropertyHandle = DetailLayoutBuilder->GetProperty(GET_MEMBER_NAME_CHECKED(UHLODLayer, CellSize), UHLODLayer::StaticClass());
-	TSharedPtr<IPropertyHandle> LoadingRangePropertyHandle = DetailLayoutBuilder->GetProperty(GET_MEMBER_NAME_CHECKED(UHLODLayer, LoadingRange), UHLODLayer::StaticClass());
-	TSharedPtr<IPropertyHandle> ParentLayerPropertyHandle = DetailLayoutBuilder->GetProperty(GET_MEMBER_NAME_CHECKED(UHLODLayer, ParentLayer), UHLODLayer::StaticClass());
+	TSharedPtr<IPropertyHandle> IsSpatiallyLoadedPropertyHandle = InDetailLayoutBuilder.GetProperty(GET_MEMBER_NAME_CHECKED(UHLODLayer, bIsSpatiallyLoaded), UHLODLayer::StaticClass());
+	TSharedPtr<IPropertyHandle> CellSizePropertyHandle = InDetailLayoutBuilder.GetProperty(GET_MEMBER_NAME_CHECKED(UHLODLayer, CellSize), UHLODLayer::StaticClass());
+	TSharedPtr<IPropertyHandle> LoadingRangePropertyHandle = InDetailLayoutBuilder.GetProperty(GET_MEMBER_NAME_CHECKED(UHLODLayer, LoadingRange), UHLODLayer::StaticClass());
+	TSharedPtr<IPropertyHandle> ParentLayerPropertyHandle = InDetailLayoutBuilder.GetProperty(GET_MEMBER_NAME_CHECKED(UHLODLayer, ParentLayer), UHLODLayer::StaticClass());
 
 	bool bShowLegacyHLODLayerProperties = false;
 	if (UWorld* EditorWorld = GEditor->GetEditorWorldContext().World())
@@ -43,16 +41,13 @@ void FWorldPartitionHLODLayerDetailsCustomization::CustomizeDetails(IDetailLayou
 	
 	if (!bShowLegacyHLODLayerProperties)
 	{
-		DetailLayoutBuilder->HideProperty(IsSpatiallyLoadedPropertyHandle);
-		DetailLayoutBuilder->HideProperty(CellSizePropertyHandle);
-		DetailLayoutBuilder->HideProperty(LoadingRangePropertyHandle);
+		InDetailLayoutBuilder.HideProperty(IsSpatiallyLoadedPropertyHandle);
+		InDetailLayoutBuilder.HideProperty(CellSizePropertyHandle);
+		InDetailLayoutBuilder.HideProperty(LoadingRangePropertyHandle);
 	}
 	else
 	{
-		IsSpatiallyLoadedPropertyHandle->SetOnPropertyValueChanged(FSimpleDelegate::CreateLambda([this]()
-		{
-			DetailLayoutBuilder->ForceRefreshDetails();
-		}));
+		IsSpatiallyLoadedPropertyHandle->SetOnPropertyValueChanged(FSimpleDelegate::CreateSP(this, &FWorldPartitionHLODLayerDetailsCustomization::ForceRefreshLayout));
 
 		PRAGMA_DISABLE_DEPRECATION_WARNINGS
 		const bool bIsSpatiallyLoaded = HLODLayer->IsSpatiallyLoaded();
@@ -60,13 +55,19 @@ void FWorldPartitionHLODLayerDetailsCustomization::CustomizeDetails(IDetailLayou
 
 		if (!bIsSpatiallyLoaded)
 		{
-			DetailLayoutBuilder->HideProperty(CellSizePropertyHandle);
-			DetailLayoutBuilder->HideProperty(LoadingRangePropertyHandle);
-			DetailLayoutBuilder->HideProperty(ParentLayerPropertyHandle);
+			InDetailLayoutBuilder.HideProperty(CellSizePropertyHandle);
+			InDetailLayoutBuilder.HideProperty(LoadingRangePropertyHandle);
+			InDetailLayoutBuilder.HideProperty(ParentLayerPropertyHandle);
 		}
 	}
 }
 
+void FWorldPartitionHLODLayerDetailsCustomization::CustomizeDetails(const TSharedPtr<IDetailLayoutBuilder>& InDetailLayoutBuilder)
+{
+	DetailLayoutBuilderWeakPtr = InDetailLayoutBuilder;
+	CustomizeDetails(*InDetailLayoutBuilder);
+}
+
 
 FWorldPartitionHLODLayerDetailsCustomization::~FWorldPartitionHLODLayerDetailsCustomization()
 {
@@ -75,7 +76,16 @@ FWorldPartitionHLODLayerDetailsCustomization::~FWorldPartitionHLODLayerDetailsCu
 
 void FWorldPartitionHLODLayerDetailsCustomization::OnWorldPartitionChanged(UWorld* InWorld)
 {
-	DetailLayoutBuilder->ForceRefreshDetails();
+	ForceRefreshLayout();
+}
+
+void FWorldPartitionHLODLayerDetailsCustomization::ForceRefreshLayout()
+{
+	IDetailLayoutBuilder* DetailLayoutBuilder = DetailLayoutBuilderWeakPtr.Pin().Get();
+	if (DetailLayoutBuilder)
+	{
+		DetailLayoutBuilder->ForceRefreshDetails();
+	}
 }
 
 #undef LOCTEXT_NAMESPACE
diff --git a/Engine/Source/Editor/WorldPartitionEditor/Private/WorldPartition/Customizations/WorldPartitionHLODLayerDetailsCustomization.h b/Engine/Source/Editor/WorldPartitionEditor/Private/WorldPartition/Customizations/WorldPartitionHLODLayerDetailsCustomization.h
--- a/Engine/Source/Editor/WorldPartitionEditor/Private/WorldPartition/Customizations/WorldPartitionHLODLayerDetailsCustomization.h
+++ b/Engine/Source/Editor/WorldPartitionEditor/Private/WorldPartition/Customizations/WorldPartitionHLODLayerDetailsCustomization.h
@@ -20,7 +20,9 @@ public:
 private:
 	/** IDetailCustomization interface */
 	virtual void CustomizeDetails(IDetailLayoutBuilder& InDetailLayoutBuilder) override;
+	virtual void CustomizeDetails(const TSharedPtr<IDetailLayoutBuilder>& InDetailLayoutBuilder) override;
 	void OnWorldPartitionChanged(UWorld* InWorld);
+	void ForceRefreshLayout();
 
-	IDetailLayoutBuilder* DetailLayoutBuilder;
+	TWeakPtr<IDetailLayoutBuilder> DetailLayoutBuilderWeakPtr;
 };
\ No newline at end of file
