Fix validation errors with /Script/Engine.HLODLayer properties
Existing PropertyEditorPermissionList in memory for /Script/Engine.HLODLayer was being replaced instead of appending a few additional properties
#rnx

#Synced-CL 50193382

diff --git a/Engine/Source/Editor/WorldPartitionEditor/Private/WorldPartition/HLOD/HLODEditorSubsystem.cpp b/Engine/Source/Editor/WorldPartitionEditor/Private/WorldPartition/HLOD/HLODEditorSubsystem.cpp
--- a/Engine/Source/Editor/WorldPartitionEditor/Private/WorldPartition/HLOD/HLODEditorSubsystem.cpp
+++ b/Engine/Source/Editor/WorldPartitionEditor/Private/WorldPartition/HLOD/HLODEditorSubsystem.cpp
@@ -203,14 +203,7 @@ void UWorldPartitionHLODEditorSubsystem::ApplyHLODSettingsFiltering()
 	{
 		for (const TPair<TSoftObjectPtr<UStruct>, TSet<FName>>& StructProperties : StructsPropertiesVisibility.FindOrAdd(EHLODSettingsVisibility::BasicSettings))
 		{
-			FNamePermissionList PermissionList;
-
-			for (FName PropertyName : StructProperties.Value)
-			{
-				PermissionList.AddAllowListItem(PropertyPermissionListOwnerName, PropertyName);
-			}
-
-			FPropertyEditorPermissionList::Get().AddPermissionList(StructProperties.Key, PermissionList, EPropertyPermissionListRules::UseExistingPermissionList, { PropertyPermissionListOwnerName });
+			FPropertyEditorPermissionList::Get().AddToAllowList(StructProperties.Key, StructProperties.Value.Array(), PropertyPermissionListOwnerName);
 		}
 	}
 }
