World Partition - HLOD: Optimized HLOD in editor visibility updates
For large worlds containing lots of HLOD, the HLOD visibility update in the editor was very costly. This update is only triggered on editor camera movements, but this still means moving around the world was not as smooth as it should have been.
Cost during camera movement, before vs after:
-> CitySample - BigCity (6061 HLOD actors): 25.1ms vs. 0.18ms (139x)
-> Engine Test - Test Scene (32161 HLOD actors): 144ms vs 1.06ms (136x)

#Synced-CL 50353017

diff --git a/Engine/Source/Editor/WorldPartitionEditor/Private/WorldPartition/HLOD/HLODEditorData.cpp b/Engine/Source/Editor/WorldPartitionEditor/Private/WorldPartition/HLOD/HLODEditorData.cpp
--- a/Engine/Source/Editor/WorldPartitionEditor/Private/WorldPartition/HLOD/HLODEditorData.cpp
+++ b/Engine/Source/Editor/WorldPartitionEditor/Private/WorldPartition/HLOD/HLODEditorData.cpp
@@ -55,7 +55,6 @@ void FWorldPartitionHLODEditorData::OnActorDescContainerInstanceRegistered(UActo
 
 		TUniquePtr<FHLODSceneNode>& HLODSceneNode = ContainerInstanceHLODActorData->HLODActorNodes.Emplace(HLODActorGuid, new FHLODSceneNode());
 		HLODSceneNode->Bounds = HLODActorDesc.GetEditorBounds().TransformBy(WorldPartition->GetInstanceTransform());
-		HLODSceneNode->HLODActorHandle = FWorldPartitionHandle(WorldPartition, HLODActorDesc.GetGuid());
 	}
 
 	// If there are no HLOD actors in this container instance, early out
@@ -313,7 +312,7 @@ void FHLODSceneNode::UpdateVisibility(const FVector& InCameraLocation, double In
 {
 	bool bNodeShouldBeVisible = !bInForceHidden;
 
-	if (AWorldPartitionHLOD* HLODActor = Cast<AWorldPartitionHLOD>(HLODActorHandle.GetActor()))
+	if (HLODActor)
 	{
 		// Do not show node if there are loaded actors under it
 		if (bNodeShouldBeVisible)
@@ -325,33 +324,30 @@ void FHLODSceneNode::UpdateVisibility(const FVector& InCameraLocation, double In
 		// Perform distance culling
 		if (bNodeShouldBeVisible)
 		{
-			// Do not perform near culling if this HLOD has no child HLOD or none of them is loaded and has all components registered
+			// Do not perform near culling if this HLOD has no child HLOD registered
 			const bool bHasLoadedAndRegisteredChildrenHLODs = Algo::AnyOf(ChildrenHLODs, [](FHLODSceneNode* ChildHLODNode)
 			{
-				if(AActor* HLODActor = ChildHLODNode->HLODActorHandle.GetActor())
-				{
-					if (HLODActor->HasActorRegisteredAllComponents())
-					{
-						return true;
-					}
-				}
-				return false;
+				return ChildHLODNode->HLODActor != nullptr;
 			});
 
-			if (bHasLoadedAndRegisteredChildrenHLODs)
+			const bool bNeedsDistanceCheck = bHasLoadedAndRegisteredChildrenHLODs || (InMinDrawDistance != 0 || InMaxDrawDistance != 0);
+			if (bNeedsDistanceCheck)
 			{
 				const double DistanceSquared = Bounds.ComputeSquaredDistanceFromBoxToPoint(InCameraLocation);
-				const bool bNearCulled = DistanceSquared < FMath::Square(HLODActor->GetMinVisibleDistance());
-				bNodeShouldBeVisible = !bNearCulled;
-			}
-			
-			// Apply HLOD min/max draw distance user setting
-			if (bNodeShouldBeVisible && (InMinDrawDistance != 0 || InMaxDrawDistance != 0))
-			{
-				const double DistanceSquared = Bounds.ComputeSquaredDistanceFromBoxToPoint(InCameraLocation);
-				const bool bNearCulled = InMinDrawDistance == 0 ? false : DistanceSquared < FMath::Square(InMinDrawDistance);
-				const bool bFarCulled = InMaxDrawDistance == 0 ? false : DistanceSquared > FMath::Square(InMaxDrawDistance);
-				bNodeShouldBeVisible = !bNearCulled && !bFarCulled;
+
+				if (bHasLoadedAndRegisteredChildrenHLODs)
+				{
+					const bool bNearCulled = DistanceSquared < FMath::Square(HLODActor->GetMinVisibleDistance());
+					bNodeShouldBeVisible = !bNearCulled;
+				}
+
+				// Apply HLOD min/max draw distance user setting
+				if (bNodeShouldBeVisible && (InMinDrawDistance != 0 || InMaxDrawDistance != 0))
+				{
+					const bool bNearCulled = InMinDrawDistance == 0 ? false : DistanceSquared < FMath::Square(InMinDrawDistance);
+					const bool bFarCulled = InMaxDrawDistance == 0 ? false : DistanceSquared > FMath::Square(InMaxDrawDistance);
+					bNodeShouldBeVisible = !bNearCulled && !bFarCulled;
+				}
 			}
 		}
 
@@ -380,3 +376,29 @@ void FWorldPartitionHLODEditorData::SetHLODLoadingState(bool bInShouldBeLoaded)
 		HLODActorsLoader.Reset();
 	}
 }
+
+void FWorldPartitionHLODEditorData::RegisterHLODActor(AWorldPartitionHLOD* Actor)
+{
+	const FGuid& ActorGuid = Actor->GetActorGuid();
+	for (auto& [ContainerInstance, HLODActorData] : PerContainerInstanceHLODActorDataMap)
+	{
+		if (TUniquePtr<FHLODSceneNode>* NodePtr = HLODActorData.HLODActorNodes.Find(ActorGuid))
+		{
+			(*NodePtr)->HLODActor = Actor;
+			return;
+		}
+	}
+}
+
+void FWorldPartitionHLODEditorData::UnregisterHLODActor(AWorldPartitionHLOD* Actor)
+{
+	const FGuid& ActorGuid = Actor->GetActorGuid();
+	for (auto& [ContainerInstance, HLODActorData] : PerContainerInstanceHLODActorDataMap)
+	{
+		if (TUniquePtr<FHLODSceneNode>* NodePtr = HLODActorData.HLODActorNodes.Find(ActorGuid))
+		{
+			(*NodePtr)->HLODActor = nullptr;
+			return;
+		}
+	}
+}
diff --git a/Engine/Source/Editor/WorldPartitionEditor/Private/WorldPartition/HLOD/HLODEditorData.h b/Engine/Source/Editor/WorldPartitionEditor/Private/WorldPartition/HLOD/HLODEditorData.h
--- a/Engine/Source/Editor/WorldPartitionEditor/Private/WorldPartition/HLOD/HLODEditorData.h
+++ b/Engine/Source/Editor/WorldPartitionEditor/Private/WorldPartition/HLOD/HLODEditorData.h
@@ -5,6 +5,7 @@
 #include "WorldPartition/HLOD/HLODLoaderAdapter.h"
 #include "ExternalDirtyActorsTracker.h"
 
+class AWorldPartitionHLOD;
 class UWorldPartition;
 class UActorDescContainerInstance;
 
@@ -19,7 +20,7 @@ struct FHLODSceneNode
 
 	FBoxSphereBounds Bounds;
 	bool bCachedIsVisible = true;
-	FWorldPartitionHandle HLODActorHandle;
+	AWorldPartitionHLOD* HLODActor = nullptr;
 	
 	int32 HasIntersectingLoadedRegion = INDEX_NONE;
 };
@@ -38,6 +39,9 @@ public:
 
 	void SetHLODLoadingState(bool bInShouldBeLoaded);
 
+	void RegisterHLODActor(AWorldPartitionHLOD* Actor);
+	void UnregisterHLODActor(AWorldPartitionHLOD* Actor);
+
 private:
 	void OnActorDescContainerInstanceRegistered(UActorDescContainerInstance* InContainerInstance);
 	void OnActorDescContainerInstanceUnregistered(UActorDescContainerInstance* InContainerInstance);
diff --git a/Engine/Source/Editor/WorldPartitionEditor/Private/WorldPartition/HLOD/HLODEditorSubsystem.cpp b/Engine/Source/Editor/WorldPartitionEditor/Private/WorldPartition/HLOD/HLODEditorSubsystem.cpp
--- a/Engine/Source/Editor/WorldPartitionEditor/Private/WorldPartition/HLOD/HLODEditorSubsystem.cpp
+++ b/Engine/Source/Editor/WorldPartitionEditor/Private/WorldPartition/HLOD/HLODEditorSubsystem.cpp
@@ -168,12 +168,18 @@ void UWorldPartitionHLODEditorSubsystem::Initialize(FSubsystemCollectionBase& Co
 	UWorldPartitionHLODsBuilder::GetOnWorldPartitionHLODBuildWithFiltersCompleted().AddUObject(this, &UWorldPartitionHLODEditorSubsystem::ForceHLODVisibilityUpdate);
 
 	ApplyHLODSettingsFiltering();
+
+	HLODActorEditorRegisteredHandle = AWorldPartitionHLOD::GetHLODActorEditorRegisteredDelegate().AddUObject(this, &UWorldPartitionHLODEditorSubsystem::OnHLODActorRegistered);
+	HLODActorEditorUnregisteredHandle = AWorldPartitionHLOD::GetHLODActorEditorUnregisteredDelegate().AddUObject(this, &UWorldPartitionHLODEditorSubsystem::OnHLODActorUnregistered);
 }
 
 void UWorldPartitionHLODEditorSubsystem::Deinitialize()
 {
 	Super::Deinitialize();
 
+	AWorldPartitionHLOD::GetHLODActorEditorRegisteredDelegate().Remove(HLODActorEditorRegisteredHandle);
+	AWorldPartitionHLOD::GetHLODActorEditorUnregisteredDelegate().Remove(HLODActorEditorUnregisteredHandle);
+
 	UWorldPartitionHLODsBuilder::GetOnWorldPartitionHLODBuildWithFiltersCompleted().RemoveAll(this);
 
 	UWorldPartitionEditorSettings::OnSettingsChanged().RemoveAll(this);
@@ -239,6 +245,34 @@ void UWorldPartitionHLODEditorSubsystem::OnLoaderAdapterStateChanged(const IWorl
 	ForceHLODStateUpdate();
 }
 
+void UWorldPartitionHLODEditorSubsystem::OnHLODActorRegistered(AWorldPartitionHLOD* Actor)
+{
+	if (UWorld* World = Actor->GetWorld())
+	{
+		if (UWorldPartition* WorldPartition = World->GetWorldPartition())
+		{
+			if (TPimplPtr<FWorldPartitionHLODEditorData>* HLODEditorData = WorldPartitionsHLODEditorData.Find(WorldPartition))
+			{
+				(*HLODEditorData)->RegisterHLODActor(Actor);
+			}
+		}
+	}
+}
+
+void UWorldPartitionHLODEditorSubsystem::OnHLODActorUnregistered(AWorldPartitionHLOD* Actor)
+{
+	if (UWorld* World = Actor->GetWorld())
+	{
+		if (UWorldPartition* WorldPartition = World->GetWorldPartition())
+		{
+			if (TPimplPtr<FWorldPartitionHLODEditorData>* HLODEditorData = WorldPartitionsHLODEditorData.Find(WorldPartition))
+			{
+				(*HLODEditorData)->UnregisterHLODActor(Actor);
+			}
+		}
+	}
+}
+
 void UWorldPartitionHLODEditorSubsystem::ForceHLODVisibilityUpdate()
 {
 	if (IsHLODInEditorEnabled())
diff --git a/Engine/Source/Editor/WorldPartitionEditor/Public/WorldPartition/HLOD/HLODEditorSubsystem.h b/Engine/Source/Editor/WorldPartitionEditor/Public/WorldPartition/HLOD/HLODEditorSubsystem.h
--- a/Engine/Source/Editor/WorldPartitionEditor/Public/WorldPartition/HLOD/HLODEditorSubsystem.h
+++ b/Engine/Source/Editor/WorldPartitionEditor/Public/WorldPartition/HLOD/HLODEditorSubsystem.h
@@ -83,7 +83,13 @@ private:
 
 	UE_API void OnColorHandlerPropertyChangedEvent(UObject* InObject, FPropertyChangedEvent& InPropertyChangedEvent);
 
+	void OnHLODActorRegistered(AWorldPartitionHLOD* Actor);
+	void OnHLODActorUnregistered(AWorldPartitionHLOD* Actor);
+
 private:
+	FDelegateHandle HLODActorEditorRegisteredHandle;
+	FDelegateHandle HLODActorEditorUnregisteredHandle;
+
 	FVector CachedCameraLocation;
 	double CachedHLODMinDrawDistance;
 	double CachedHLODMaxDrawDistance;
diff --git a/Engine/Source/Runtime/Engine/Private/WorldPartition/HLOD/HLODActor.cpp b/Engine/Source/Runtime/Engine/Private/WorldPartition/HLOD/HLODActor.cpp
--- a/Engine/Source/Runtime/Engine/Private/WorldPartition/HLOD/HLODActor.cpp
+++ b/Engine/Source/Runtime/Engine/Private/WorldPartition/HLOD/HLODActor.cpp
@@ -1,8 +1,10 @@
 // Copyright Epic Games, Inc. All Rights Reserved.
 
 #include "WorldPartition/HLOD/HLODActor.h"
+#include "Engine/Engine.h"
 #include "Engine/World.h"
 #include "Engine/StaticMesh.h"
+#include "EngineUtils.h"
 #include "WorldPartition/HLOD/HLODRuntimeSubsystem.h"
 #include "Components/PrimitiveComponent.h"
 #include "Components/InstancedStaticMeshComponent.h"
@@ -53,6 +55,8 @@ static FAutoConsoleVariableRef CVarWorldPartitionHLODForceDisableShadows(
 #if WITH_EDITOR
 const FName AWorldPartitionHLOD::NAME_HLODHash_AssetTag(TEXT("HLODActor_HLODHash"));
 FWorldPartitionHLODBuildEventDelegate AWorldPartitionHLOD::HLODBuildEventDelegate;
+FWorldPartitionHLODActorEditorRegisteredDelegate AWorldPartitionHLOD::HLODActorEditorRegisteredDelegate;
+FWorldPartitionHLODActorEditorUnregisteredDelegate AWorldPartitionHLOD::HLODActorEditorUnregisteredDelegate;
 #endif
 
 AWorldPartitionHLOD::AWorldPartitionHLOD(const FObjectInitializer& ObjectInitializer)
@@ -431,6 +435,34 @@ void AWorldPartitionHLOD::PreRegisterAllComponents()
 	}
 }
 
+void AWorldPartitionHLOD::PostRegisterAllComponents()
+{
+	Super::PostRegisterAllComponents();
+#if WITH_EDITOR
+	if (UWorld* World = GetWorld())
+	{
+		if (World->IsEditorWorld())
+		{
+			HLODActorEditorRegisteredDelegate.Broadcast(this);
+		}
+	}
+#endif
+}
+
+void AWorldPartitionHLOD::PostUnregisterAllComponents()
+{
+#if WITH_EDITOR
+	if (UWorld* World = GetWorld())
+	{
+		if (World->IsEditorWorld())
+		{
+			HLODActorEditorUnregisteredDelegate.Broadcast(this);
+		}
+	}
+#endif
+	Super::PostUnregisterAllComponents();
+}
+
 #if WITH_EDITOR
 void AWorldPartitionHLOD::RerunConstructionScripts()
 {}
@@ -639,6 +671,16 @@ FWorldPartitionHLODBuildEventDelegate& AWorldPartitionHLOD::GetHLODBuildEventDel
 	return HLODBuildEventDelegate;
 }
 
+FWorldPartitionHLODActorEditorRegisteredDelegate& AWorldPartitionHLOD::GetHLODActorEditorRegisteredDelegate()
+{
+	return HLODActorEditorRegisteredDelegate;
+}
+
+FWorldPartitionHLODActorEditorUnregisteredDelegate& AWorldPartitionHLOD::GetHLODActorEditorUnregisteredDelegate()
+{
+	return HLODActorEditorUnregisteredDelegate;
+}
+
 const TCHAR* HLOD_REPORT_BEGIN = TEXT("### HLOD_REPORT_BEGIN ###");
 const TCHAR* HLODLastBuildInfoSectionHeader = TEXT("### Last Build Info ###");
 const TCHAR* HLODActorDetailsSectionHeader = TEXT("### HLOD Actor Details ###");
diff --git a/Engine/Source/Runtime/Engine/Public/WorldPartition/HLOD/HLODActor.h b/Engine/Source/Runtime/Engine/Public/WorldPartition/HLOD/HLODActor.h
--- a/Engine/Source/Runtime/Engine/Public/WorldPartition/HLOD/HLODActor.h
+++ b/Engine/Source/Runtime/Engine/Public/WorldPartition/HLOD/HLODActor.h
@@ -38,6 +38,8 @@ struct FWorldPartitionHLODBuildEvent
 };
 
 DECLARE_MULTICAST_DELEGATE_OneParam(FWorldPartitionHLODBuildEventDelegate, const FWorldPartitionHLODBuildEvent&)
+DECLARE_MULTICAST_DELEGATE_OneParam(FWorldPartitionHLODActorEditorRegisteredDelegate, AWorldPartitionHLOD*)
+DECLARE_MULTICAST_DELEGATE_OneParam(FWorldPartitionHLODActorEditorUnregisteredDelegate, AWorldPartitionHLOD*)
 #endif
 
 UENUM(BlueprintType)
@@ -129,6 +131,8 @@ public:
 	void ResetStats() { HLODStats.Reset(); }
 
 	ENGINE_API static FWorldPartitionHLODBuildEventDelegate& GetHLODBuildEventDelegate();
+	ENGINE_API static FWorldPartitionHLODActorEditorRegisteredDelegate& GetHLODActorEditorRegisteredDelegate();
+	ENGINE_API static FWorldPartitionHLODActorEditorUnregisteredDelegate& GetHLODActorEditorUnregisteredDelegate();
 
 	UFUNCTION(BlueprintCallable, Category = "World Partition", meta = (DevelopmentOnly))
 	ENGINE_API TArray<UObject*> ExportHLODAssets(const struct FExportHLODAssetsParams& ExportHLODAssetsParams, FString& OutErrorMessage) const;
@@ -169,6 +173,8 @@ protected:
 
 	//~ Begin AActor Interface.
 	ENGINE_API virtual void PreRegisterAllComponents() override;
+	ENGINE_API virtual void PostRegisterAllComponents() override;
+	ENGINE_API virtual void PostUnregisterAllComponents() override;
 	ENGINE_API virtual void BeginPlay() override;
 	ENGINE_API virtual void EndPlay(const EEndPlayReason::Type EndPlayReason) override;
 	virtual bool SupportsIncrementalPreRegisterComponents() const override
@@ -258,6 +264,8 @@ private:
 
 #if WITH_EDITOR
 	static FWorldPartitionHLODBuildEventDelegate HLODBuildEventDelegate;
+	static FWorldPartitionHLODActorEditorRegisteredDelegate HLODActorEditorRegisteredDelegate;
+	static FWorldPartitionHLODActorEditorUnregisteredDelegate HLODActorEditorUnregisteredDelegate;
 #endif
 };
 
