World Partition - HLOD: Added option to choose the editor loading behavior of HLOD actors belonging to a given HLOD Layer

Add EHLODLayerEditorLoadingBehavior enum to UHLODLayer with three options:
- Default: Use existing behavior (Instancing not loaded except if non-spatial, others loaded)
- Always Loaded: Force load HLODs from this layer in editor
- Never Loaded: Force not load HLODs from this layer in editor

This preserves backward compatibility while giving users explicit control.

#Synced-CL 50353017

diff --git a/Engine/Source/Editor/WorldPartitionEditor/Private/WorldPartition/HLOD/HLODLoaderAdapter.cpp b/Engine/Source/Editor/WorldPartitionEditor/Private/WorldPartition/HLOD/HLODLoaderAdapter.cpp
--- a/Engine/Source/Editor/WorldPartitionEditor/Private/WorldPartition/HLOD/HLODLoaderAdapter.cpp
+++ b/Engine/Source/Editor/WorldPartitionEditor/Private/WorldPartition/HLOD/HLODLoaderAdapter.cpp
@@ -26,37 +26,35 @@ void FLoaderAdapterHLOD::ForEachActor(TFunctionRef<void(const FWorldPartitionHan
 
 	for (FActorDescContainerInstanceCollection::TIterator<AWorldPartitionHLOD> HLODIterator(WorldPartition); HLODIterator; ++HLODIterator)
 	{
-		const FHLODActorDesc& HLODActorDesc = *(FHLODActorDesc*)HLODIterator->GetActorDesc();
-		if (ShouldLoadHLOD(HLODActorDesc))
-		{
-			InOperation(FWorldPartitionHandle(WorldPartition, HLODActorDesc.GetGuid()));
-		}
+		InOperation(FWorldPartitionHandle(WorldPartition, HLODIterator->GetGuid()));
 	}
 }
 
-// TODO - How to choose which HLODs are loaded ?
-// * Always loaded should be loaded...
-// * Should be a user setting/bookmark setting maybe?
-// * World settings? 
-// * HLOD Layer setting?
-//
-// Current solution is to load always loaded HLOD + any HLOD not built from instancing
-bool FLoaderAdapterHLOD::ShouldLoadHLOD(const FHLODActorDesc& HLODActorDesc) const
+bool FLoaderAdapterHLOD::PassActorDescFilter(const FWorldPartitionHandle& ActorHandle) const
 {
-	FSoftObjectPath HLODLayerPath(HLODActorDesc.GetSourceHLODLayer());
-	if (UHLODLayer* HLODLayer = Cast<UHLODLayer>(HLODLayerPath.TryLoad()))
+	// Do not call the base class implementation of PassActorDescFilter() here, as it will skip actors which are not editor relevant
+
+	const FHLODActorDesc* HLODActorDesc = static_cast<const FHLODActorDesc*>(ActorHandle->GetActorDesc());
+	if (HLODActorDesc)
 	{
-		if (!HLODActorDesc.GetIsSpatiallyLoaded() || HLODLayer->GetLayerType() != EHLODLayerType::Instancing)
+		FSoftObjectPath HLODLayerPath(HLODActorDesc->GetSourceHLODLayer());
+		if (UHLODLayer* HLODLayer = Cast<UHLODLayer>(HLODLayerPath.TryLoad()))
 		{
-			return true;
+			switch (HLODLayer->GetEditorLoadingBehavior())
+			{
+			case EHLODLayerEditorLoadingBehavior::ForceLoaded:
+				return true;
+			case EHLODLayerEditorLoadingBehavior::ForceNotLoaded:
+				return false;
+			case EHLODLayerEditorLoadingBehavior::Default:
+			default:
+				if (!HLODActorDesc->GetIsSpatiallyLoaded() || HLODLayer->GetLayerType() != EHLODLayerType::Instancing)
+				{
+					return true;
+				}
+			}
 		}
 	}
 
 	return false;
 }
-
-bool FLoaderAdapterHLOD::PassActorDescFilter(const FWorldPartitionHandle& ActorHandle) const
-{
-	// Avoid the base class implementation which will skip actors which are not editor relevant
-	return true;
-}
diff --git a/Engine/Source/Editor/WorldPartitionEditor/Public/WorldPartition/HLOD/HLODLoaderAdapter.h b/Engine/Source/Editor/WorldPartitionEditor/Public/WorldPartition/HLOD/HLODLoaderAdapter.h
--- a/Engine/Source/Editor/WorldPartitionEditor/Public/WorldPartition/HLOD/HLODLoaderAdapter.h
+++ b/Engine/Source/Editor/WorldPartitionEditor/Public/WorldPartition/HLOD/HLODLoaderAdapter.h
@@ -16,7 +16,4 @@ protected:
 	virtual void ForEachActor(TFunctionRef<void(const FWorldPartitionHandle&)> InOperation) const override;
 	virtual bool PassActorDescFilter(const FWorldPartitionHandle& ActorHandle) const override;
 	//~ End ILoaderAdapter interface
-
-private:
-	bool ShouldLoadHLOD(const FHLODActorDesc& HLODActorDesc) const;	
 };
diff --git a/Engine/Source/Runtime/Engine/Private/WorldPartition/HLOD/HLODLayer.cpp b/Engine/Source/Runtime/Engine/Private/WorldPartition/HLOD/HLODLayer.cpp
--- a/Engine/Source/Runtime/Engine/Private/WorldPartition/HLOD/HLODLayer.cpp
+++ b/Engine/Source/Runtime/Engine/Private/WorldPartition/HLOD/HLODLayer.cpp
@@ -8,6 +8,7 @@
 #include "WorldPartition/HLOD/HLODActor.h"
 
 #if WITH_EDITOR
+#include "Dialogs/Dialogs.h"
 #include "Engine/World.h"
 #include "Framework/Notifications/NotificationManager.h"
 #include "Misc/ConfigCacheIni.h"
@@ -21,6 +22,7 @@
 #include "WorldPartition/HLOD/IWorldPartitionHLODUtilities.h"
 #include "WorldPartition/HLOD/IWorldPartitionHLODUtilitiesModule.h"
 #include "WorldPartition/WorldPartition.h"
+#include "WorldPartition/WorldPartitionActorLoaderInterface.h"
 #endif
 
 #include UE_INLINE_GENERATED_CPP_BY_NAME(HLODLayer)
@@ -36,6 +38,7 @@ UHLODLayer::UHLODLayer(const FObjectInitializer& ObjectInitializer)
 	, CellSize(25600)
 	, LoadingRange(51200)
 	, HLODActorClass(AWorldPartitionHLOD::StaticClass())
+	, EditorLoadingBehavior(EHLODLayerEditorLoadingBehavior::Default)
 #endif
 {
 }
@@ -221,6 +224,30 @@ void UHLODLayer::PostEditChangeProperty(FPropertyChangedEvent& PropertyChangedEv
 			FSlateNotificationManager::Get().AddNotification(Info);
 		}
 	}
+	else if (PropertyName == GET_MEMBER_NAME_CHECKED(UHLODLayer, EditorLoadingBehavior))
+	{
+		// Warn when enabling loading for Instancing layers, as this can have performance/memory implications
+		if (EditorLoadingBehavior == EHLODLayerEditorLoadingBehavior::ForceLoaded && LayerType == EHLODLayerType::Instancing)
+		{
+			FSuppressableWarningDialog::FSetupInfo Info(
+				LOCTEXT("InstancingHLODLoadWarning_Message", "Loading Instancing HLOD actors in the editor can have significant performance and memory implications, as this may lead to a lot of assets being loaded.\n\nAre you sure you want to enable editor loading for this HLOD layer?"),
+				LOCTEXT("InstancingHLODLoadWarning_Title", "Warning: Loading Instancing HLOD Layer"),
+				TEXT("InstancingHLODLayerEditorLoadingWarning")
+			);
+			Info.ConfirmText = LOCTEXT("InstancingHLODLoadWarning_Confirm", "Yes");
+			Info.CancelText = LOCTEXT("InstancingHLODLoadWarning_Cancel", "No");
+
+			FSuppressableWarningDialog WarningDialog(Info);
+			if (WarningDialog.ShowModal() == FSuppressableWarningDialog::Cancel)
+			{
+				// User cancelled, revert to default
+				EditorLoadingBehavior = EHLODLayerEditorLoadingBehavior::Default;
+			}
+		}
+
+		// Refresh all loaders to re-evaluate which HLODs should be loaded
+		IWorldPartitionActorLoaderInterface::RefreshLoadedState(true);
+	}
 }
 
 EDataValidationResult UHLODLayer::IsDataValid(class FDataValidationContext& Context) const
diff --git a/Engine/Source/Runtime/Engine/Public/WorldPartition/HLOD/HLODLayer.h b/Engine/Source/Runtime/Engine/Public/WorldPartition/HLOD/HLODLayer.h
--- a/Engine/Source/Runtime/Engine/Public/WorldPartition/HLOD/HLODLayer.h
+++ b/Engine/Source/Runtime/Engine/Public/WorldPartition/HLOD/HLODLayer.h
@@ -38,6 +38,17 @@ enum class EHLODLayerType : uint8
 	CustomHLODActor			UMETA(DisplayName = "Custom HLOD Actor"),
 };
 
+UENUM()
+enum class EHLODLayerEditorLoadingBehavior : uint8
+{
+	/** Use default behavior: Instancing layers are not loaded (except if non-spatial), other layer types are loaded */
+	Default					UMETA(DisplayName = "Default"),
+	/** Always load HLODs from this layer in the editor */
+	ForceLoaded				UMETA(DisplayName = "Always Loaded"),
+	/** Never load HLODs from this layer in the editor */
+	ForceNotLoaded			UMETA(DisplayName = "Never Loaded"),
+};
+
 UCLASS(Blueprintable, MinimalAPI)
 class UHLODLayer : public UObject
 {
@@ -61,6 +72,7 @@ public:
 	UHLODLayer* GetParentLayer() const { return ParentLayer; }
 	void SetParentLayer(UHLODLayer* InParentLayer) { ParentLayer = InParentLayer; }
 	UHLODLayer* GetLinkedLayer() const { return LinkedLayer; }
+	EHLODLayerEditorLoadingBehavior GetEditorLoadingBehavior() const { return EditorLoadingBehavior; }
 	ENGINE_API bool DoesRequireWarmup() const;
 	ENGINE_API void ComputeHLODHash(FHLODHashBuilder& InHLODHashBuilder) const;
 	ENGINE_API bool MustBeLoadedOnServer() const;
@@ -142,6 +154,15 @@ private:
 	UPROPERTY(EditAnywhere, Category = General, AdvancedDisplay)
 	uint32 bForceRayTracingFarField : 1;
 
+	/**
+	 * Controls whether HLODs from this layer are loaded in the editor.
+	 * - Default: Instancing layers are not loaded (except if non-spatial), other layer types are loaded
+	 * - Always Loaded: Force load HLODs from this layer in editor
+	 * - Never Loaded: Prevent loading HLODs from this layer in editor
+	 */
+	UPROPERTY(EditAnywhere, Category = General, AdvancedDisplay)
+	EHLODLayerEditorLoadingBehavior EditorLoadingBehavior;
+
 	friend class FWorldPartitionHLODLayerDetailsCustomization;
 
 private:
