[WorldPartition] Fix callback for dynamic loading range editing so that modified maps get dirtied

#Synced-CL 50613390

diff --git a/Engine/Source/Editor/WorldPartitionEditor/Private/WorldPartition/Customizations/WorldPartitionDetailsCustomization.cpp b/Engine/Source/Editor/WorldPartitionEditor/Private/WorldPartition/Customizations/WorldPartitionDetailsCustomization.cpp
--- a/Engine/Source/Editor/WorldPartitionEditor/Private/WorldPartition/Customizations/WorldPartitionDetailsCustomization.cpp
+++ b/Engine/Source/Editor/WorldPartitionEditor/Private/WorldPartition/Customizations/WorldPartitionDetailsCustomization.cpp
@@ -199,7 +199,7 @@ void FWorldPartitionDetails::CustomizeDetails(IDetailLayoutBuilder& InDetailBuil
 				PropertyRow->ToolTip(LOCTEXT("DynamicLoadingRangeScaling_ToolTip", "Locally force disable the dynamic loading range scaling - stored in local per-user settings."));
 
 				TSharedPtr<IPropertyHandle> PropertyHandle = PropertyRow->GetPropertyHandle();
-				PropertyHandle->SetOnPropertyValuePreChange(FSimpleDelegate::CreateLambda([this]()
+				PropertyHandle->SetOnChildPropertyValuePreChange(FSimpleDelegate::CreateLambda([this]()
 				{
 					if (WorldPartition.IsValid() && WorldPartition->GetWorld())
 					{
