World Partition - HLOD: Transformed the HLODEditorSubsystem from a world subsystem to an editor subsystem
* Advanced HLOD settings filtering is not reapplied upon each world loading anymore.
* Moved HLOD viewport menu logic out of the LevelEditor and into the HLODEditorSubsystem

#Synced-CL 51228359

diff --git a/Engine/Source/Editor/LevelEditor/Private/ViewportToolbar/LevelEditorSubmenus.cpp b/Engine/Source/Editor/LevelEditor/Private/ViewportToolbar/LevelEditorSubmenus.cpp
--- a/Engine/Source/Editor/LevelEditor/Private/ViewportToolbar/LevelEditorSubmenus.cpp
+++ b/Engine/Source/Editor/LevelEditor/Private/ViewportToolbar/LevelEditorSubmenus.cpp
@@ -32,11 +32,9 @@
 #include "ViewportToolbar/UnrealEdViewportToolbar.h"
 #include "VirtualTextureVisualizationMenuCommands.h"
 #include "Widgets/Input/SNumericEntryBox.h"
-#include "Widgets/Input/SSpinBox.h"
 #include "Widgets/Input/SVolumeControl.h"
 #include "Widgets/SBoxPanel.h"
 #include "Widgets/SToolTip.h"
-#include "WorldPartition/IWorldPartitionEditorModule.h"
 #include "VirtualShadowMapVisualizationData.h"
 #include "Tools/EdModeInteractiveToolsContext.h"
 #include "RayTracingDebugVisualizationMenuCommands.h"
@@ -1641,293 +1639,6 @@ FToolMenuEntry CreateShowFoliageSubmenu()
 	);
 }
 
-FToolMenuEntry CreateShowHLODsSubmenu()
-{
-	// This is a dynamic entry so we can skip adding the submenu if the context
-	// indicates that the viewport's world isn't partitioned.
-	return FToolMenuEntry::InitDynamicEntry(
-		"ShowHLODsDynamic",
-		FNewToolMenuSectionDelegate::CreateLambda(
-			[](FToolMenuSection& InDynamicSection) -> void
-			{
-				TSharedPtr<::SLevelViewport> Viewport = ULevelViewportContext::GetLevelViewport(InDynamicSection);
-				if (!Viewport)
-				{
-					return;
-				}
-
-				UWorld* World = Viewport->GetWorld();
-				if (!World)
-				{
-					return;
-				}
-
-				// Only add this submenu for partitioned worlds.
-				if (!World->IsPartitionedWorld())
-				{
-					return;
-				}
-
-				InDynamicSection.AddSubMenu(
-					"ShowHLODsMenu",
-					LOCTEXT("ShowHLODsMenu", "HLODs"),
-					LOCTEXT("ShowHLODsMenu_ToolTip", "Settings for HLODs in editor"),
-					FNewToolMenuDelegate::CreateLambda(
-						[](UToolMenu* Submenu)
-						{
-							TSharedPtr<::SLevelViewport> Viewport = ULevelViewportContext::GetLevelViewport(Submenu);
-							if (!Viewport)
-							{
-								return;
-							}
-
-							UWorld* World = Viewport->GetWorld();
-							UWorldPartition* WorldPartition = World ? World->GetWorldPartition() : nullptr;
-							if (!WorldPartition)
-							{
-								return;
-							}
-
-							IWorldPartitionEditorModule* WorldPartitionEditorModule =
-								FModuleManager::GetModulePtr<IWorldPartitionEditorModule>("WorldPartitionEditor");
-							if (!WorldPartitionEditorModule)
-							{
-								return;
-							}
-
-							FText HLODInEditorDisallowedReason;
-							const bool bHLODInEditorAllowed =
-								WorldPartitionEditorModule->IsHLODInEditorAllowed(World, &HLODInEditorDisallowedReason);
-
-							// Show HLODs
-							{
-								FToolUIAction UIAction;
-								UIAction.ExecuteAction = FToolMenuExecuteAction::CreateLambda(
-									[WorldPartitionEditorModule](const FToolMenuContext& InContext)
-									{
-										WorldPartitionEditorModule->SetShowHLODsInEditor(
-											!WorldPartitionEditorModule->GetShowHLODsInEditor()
-										);
-									}
-								);
-								UIAction.CanExecuteAction = FToolMenuCanExecuteAction::CreateLambda(
-									[bHLODInEditorAllowed](const FToolMenuContext& InContext)
-									{
-										return bHLODInEditorAllowed;
-									}
-								);
-								UIAction.GetActionCheckState = FToolMenuGetActionCheckState::CreateLambda(
-									[WorldPartitionEditorModule](const FToolMenuContext& InContext)
-									{
-										return WorldPartitionEditorModule->GetShowHLODsInEditor()
-												 ? ECheckBoxState::Checked
-												 : ECheckBoxState::Unchecked;
-									}
-								);
-								FToolMenuEntry MenuEntry = FToolMenuEntry::InitMenuEntry(
-									"ShowHLODs",
-									LOCTEXT("ShowHLODs", "Show HLODs"),
-									bHLODInEditorAllowed ? LOCTEXT("ShowHLODsToolTip", "Show/Hide HLODs")
-														 : HLODInEditorDisallowedReason,
-									FSlateIcon(),
-									UIAction,
-									EUserInterfaceActionType::ToggleButton
-								);
-								Submenu->AddMenuEntry(NAME_None, MenuEntry);
-							}
-
-							// Show HLODs Over Loaded Regions
-							{
-								FToolUIAction UIAction;
-								UIAction.ExecuteAction = FToolMenuExecuteAction::CreateLambda(
-									[WorldPartitionEditorModule](const FToolMenuContext& InContext)
-									{
-										WorldPartitionEditorModule->SetShowHLODsOverLoadedRegions(
-											!WorldPartitionEditorModule->GetShowHLODsOverLoadedRegions()
-										);
-									}
-								);
-								UIAction.CanExecuteAction = FToolMenuCanExecuteAction::CreateLambda(
-									[bHLODInEditorAllowed](const FToolMenuContext& InContext)
-									{
-										return bHLODInEditorAllowed;
-									}
-								);
-								UIAction.GetActionCheckState = FToolMenuGetActionCheckState::CreateLambda(
-									[WorldPartitionEditorModule](const FToolMenuContext& InContext)
-									{
-										return WorldPartitionEditorModule->GetShowHLODsOverLoadedRegions()
-												 ? ECheckBoxState::Checked
-												 : ECheckBoxState::Unchecked;
-									}
-								);
-								FToolMenuEntry ShowHLODsEntry = FToolMenuEntry::InitMenuEntry(
-									"ShowHLODsOverLoadedRegions",
-									LOCTEXT("ShowHLODsOverLoadedRegions", "Show HLODs Over Loaded Regions"),
-									bHLODInEditorAllowed
-										? LOCTEXT("ShowHLODsOverLoadedRegions_ToolTip", "Show/Hide HLODs over loaded actors or regions")
-										: HLODInEditorDisallowedReason,
-									FSlateIcon(),
-									UIAction,
-									EUserInterfaceActionType::ToggleButton
-								);
-								Submenu->AddMenuEntry(NAME_None, ShowHLODsEntry);
-							}
-
-							// Min/Max Draw Distance
-							{
-								const double MinDrawDistanceMinValue = 0;
-								const double MinDrawDistanceMaxValue = 102400;
-
-								const double MaxDrawDistanceMinValue = 0;
-								const double MaxDrawDistanceMaxValue = 1638400;
-
-								// double SLevelViewportToolBar::OnGetHLODInEditorMinDrawDistanceValue() const
-								auto OnGetHLODInEditorMinDrawDistanceValue = []() -> double
-								{
-									IWorldPartitionEditorModule* WorldPartitionEditorModule =
-										FModuleManager::GetModulePtr<IWorldPartitionEditorModule>("WorldPartitionEditor");
-									return WorldPartitionEditorModule
-											 ? WorldPartitionEditorModule->GetHLODInEditorMinDrawDistance()
-											 : 0;
-								};
-
-								// void SLevelViewportToolBar::OnHLODInEditorMinDrawDistanceValueChanged(double NewValue) const
-								auto OnHLODInEditorMinDrawDistanceValueChanged = [](double NewValue) -> void
-								{
-									IWorldPartitionEditorModule* WorldPartitionEditorModule =
-										FModuleManager::GetModulePtr<IWorldPartitionEditorModule>("WorldPartitionEditor");
-									if (WorldPartitionEditorModule)
-									{
-										WorldPartitionEditorModule->SetHLODInEditorMinDrawDistance(NewValue);
-										GEditor->RedrawLevelEditingViewports(true);
-									}
-								};
-
-								TSharedRef<SSpinBox<double>> MinDrawDistanceSpinBox =
-									SNew(SSpinBox<double>)
-										.MinValue(MinDrawDistanceMinValue)
-										.MaxValue(MinDrawDistanceMaxValue)
-										.IsEnabled(bHLODInEditorAllowed)
-										.Value_Lambda(OnGetHLODInEditorMinDrawDistanceValue)
-										.OnValueChanged_Lambda(OnHLODInEditorMinDrawDistanceValueChanged)
-										.ToolTipText(
-											bHLODInEditorAllowed
-												? LOCTEXT(
-													  "HLODsInEditor_MinDrawDistance_Tooltip",
-													  "Sets the minimum distance at which HLOD will be rendered"
-												  )
-												: HLODInEditorDisallowedReason
-										)
-										.OnBeginSliderMovement_Lambda(
-											[]()
-											{
-												// Disable Slate throttling during slider drag to ensure immediate updates while moving the slider.
-												FSlateThrottleManager::Get().DisableThrottle(true);
-											}
-										)
-										.OnEndSliderMovement_Lambda(
-											[](double)
-											{
-												FSlateThrottleManager::Get().DisableThrottle(false);
-											}
-										);
-
-								// double SLevelViewportToolBar::OnGetHLODInEditorMaxDrawDistanceValue() const
-								auto OnGetHLODInEditorMaxDrawDistanceValue = []() -> double
-								{
-									IWorldPartitionEditorModule* WorldPartitionEditorModule =
-										FModuleManager::GetModulePtr<IWorldPartitionEditorModule>("WorldPartitionEditor");
-									return WorldPartitionEditorModule
-											 ? WorldPartitionEditorModule->GetHLODInEditorMaxDrawDistance()
-											 : 0;
-								};
-
-								// void SLevelViewportToolBar::OnHLODInEditorMaxDrawDistanceValueChanged(double NewValue) const
-								auto OnHLODInEditorMaxDrawDistanceValueChanged = [](double NewValue) -> void
-								{
-									IWorldPartitionEditorModule* WorldPartitionEditorModule =
-										FModuleManager::GetModulePtr<IWorldPartitionEditorModule>("WorldPartitionEditor");
-									if (WorldPartitionEditorModule)
-									{
-										WorldPartitionEditorModule->SetHLODInEditorMaxDrawDistance(NewValue);
-										GEditor->RedrawLevelEditingViewports(true);
-									}
-								};
-
-								TSharedRef<SSpinBox<double>> MaxDrawDistanceSpinBox =
-									SNew(SSpinBox<double>)
-										.MinValue(MaxDrawDistanceMinValue)
-										.MaxValue(MaxDrawDistanceMaxValue)
-										.IsEnabled(bHLODInEditorAllowed)
-										.Value_Lambda(OnGetHLODInEditorMaxDrawDistanceValue)
-										.OnValueChanged_Lambda(OnHLODInEditorMaxDrawDistanceValueChanged)
-										.ToolTipText(
-											bHLODInEditorAllowed
-												? LOCTEXT(
-													  "HLODsInEditor_MaxDrawDistance_Tooltip", "Sets the maximum distance at which HLODs will be rendered (0.0 means infinite)"
-												  )
-												: HLODInEditorDisallowedReason
-										)
-										.OnBeginSliderMovement_Lambda(
-											[]()
-											{
-												// Disable Slate throttling during slider drag to ensure immediate updates while moving the slider.
-												FSlateThrottleManager::Get().DisableThrottle(true);
-											}
-										)
-										.OnEndSliderMovement_Lambda(
-											[](double)
-											{
-												FSlateThrottleManager::Get().DisableThrottle(false);
-											}
-										);
-
-								auto CreateDrawDistanceWidget = [](TSharedRef<SSpinBox<double>> InSpinBoxWidget)
-								{
-									// clang-format off
-									return SNew(SBox)
-										.HAlign(HAlign_Right)
-										[
-											SNew(SBox)
-										  .Padding(FMargin(0.0f, 0.0f, 0.0f, 0.0f))
-										  .WidthOverride(100.0f)
-											[
-												SNew(SBorder)
-												.BorderImage(FAppStyle::Get().GetBrush("Menu.WidgetBorder"))
-												.Padding(FMargin(1.0f))
-												[
-													InSpinBoxWidget
-												]
-											]
-										];
-									// clang-format on
-								};
-
-								FToolMenuEntry MinDrawDistanceMenuEntry = FToolMenuEntry::InitWidget(
-									"Min Draw Distance",
-									CreateDrawDistanceWidget(MinDrawDistanceSpinBox),
-									LOCTEXT("MinDrawDistance", "Min Draw Distance")
-								);
-								Submenu->AddMenuEntry(NAME_None, MinDrawDistanceMenuEntry);
-
-								FToolMenuEntry MaxDrawDistanceMenuEntry = FToolMenuEntry::InitWidget(
-									"Max Draw Distance",
-									CreateDrawDistanceWidget(MaxDrawDistanceSpinBox),
-									LOCTEXT("MaxDrawDistance", "Max Draw Distance")
-								);
-								Submenu->AddMenuEntry(NAME_None, MaxDrawDistanceMenuEntry);
-							}
-						}
-					),
-					false,
-					FSlateIcon(FAppStyle::Get().GetStyleSetName(), "ShowFlagsMenu.SubMenu.HLODs")
-				);
-			}
-		)
-	);
-}
-
 FToolMenuEntry CreateShowLayersSubmenu()
 {
 	// This is a dynamic entry so we can skip adding the submenu if the context
@@ -2228,13 +1939,6 @@ FToolMenuEntry CreateShowSubmenu()
 					AllShowFlagsSection.AddEntry(ShowFoliageSubmenu);
 				}
 
-				// Show HLODs
-				{
-					FToolMenuEntry ShowHLODSubmenu = CreateShowHLODsSubmenu();
-					ShowHLODSubmenu.InsertPosition.Position = EToolMenuInsertType::First;
-					AllShowFlagsSection.AddEntry(ShowHLODSubmenu);
-				}
-
 				// Show Layers
 				{
 					FToolMenuEntry ShowLayersSubmenu = CreateShowLayersSubmenu();
diff --git a/Engine/Source/Editor/WorldPartitionEditor/Private/WorldPartition/HLOD/HLODEditorSubsystem.cpp b/Engine/Source/Editor/WorldPartitionEditor/Private/WorldPartition/HLOD/HLODEditorSubsystem.cpp
--- a/Engine/Source/Editor/WorldPartitionEditor/Private/WorldPartition/HLOD/HLODEditorSubsystem.cpp
+++ b/Engine/Source/Editor/WorldPartitionEditor/Private/WorldPartition/HLOD/HLODEditorSubsystem.cpp
@@ -6,14 +6,17 @@
 #include "Engine/Engine.h"
 #include "Engine/World.h"
 #include "GameFramework/ActorPrimitiveColorHandler.h"
-#include "GameFramework/WorldSettings.h"
 #include "LevelInstance/LevelInstanceInterface.h"
 #include "LevelInstance/LevelInstanceSubsystem.h"
+#include "LevelEditor.h"
+#include "LevelEditorMenuContext.h"
 #include "Misc/FileHelper.h"
 #include "Misc/StringOutputDevice.h"
+#include "SLevelViewport.h"
 #include "StaticMeshResources.h"
 #include "Subsystems/UnrealEditorSubsystem.h"
-#include "UObject/UObjectIterator.h"
+#include "ToolMenus.h"
+#include "Widgets/Input/SSpinBox.h"
 #include "WorldPartitionEditorModule.h"
 #include "WorldPartition/HLOD/HLODActor.h"
 #include "WorldPartition/HLOD/HLODActorDesc.h"
@@ -25,7 +28,6 @@
 #include "WorldPartition/WorldPartitionEditorSettings.h"
 #include "WorldPartition/WorldPartitionHLODsBuilder.h"
 #include "WorldPartition/WorldPartitionSubsystem.h"
-#include "WorldPartition/DataLayer/WorldDataLayers.h"
 #include "WorldPartition/DataLayer/DataLayerManager.h"
 
 #include "PropertyPermissionList.h"
@@ -131,24 +133,22 @@ bool UWorldPartitionHLODEditorSubsystem::IsHLODInEditorEnabled()
 	{
 		return false;
 	}
+
+	if (!GEditor)
+	{
+		return false;
+	}
 	
+	UWorld* CurrentWorld = GEditor->GetEditorWorldContext().World();
 	const IWorldPartitionEditorModule* WorldPartitionEditorModule = FModuleManager::GetModulePtr<IWorldPartitionEditorModule>("WorldPartitionEditor");
-	const bool bShowHLODsInEditorForWorld = WorldPartitionEditorModule && WorldPartitionEditorModule->IsHLODInEditorAllowed(GetWorld());
+	const bool bShowHLODsInEditorForWorld = WorldPartitionEditorModule && WorldPartitionEditorModule->IsHLODInEditorAllowed(CurrentWorld);
 	const bool bShowHLODsInEditorUserSetting = WorldPartitionEditorModule && WorldPartitionEditorModule->GetShowHLODsInEditor();
 	const bool bWorldPartitionLoadingInEditorEnabled = WorldPartitionEditorModule && WorldPartitionEditorModule->GetEnableLoadingInEditor();
 	return CVarHLODInEditorEnabled.GetValueOnGameThread() && bShowHLODsInEditorUserSetting && bShowHLODsInEditorForWorld && bWorldPartitionLoadingInEditorEnabled;
 }
 
-bool UWorldPartitionHLODEditorSubsystem::DoesSupportWorldType(const EWorldType::Type WorldType) const
-{
-	return WorldType == EWorldType::Editor;
-}
-
 void UWorldPartitionHLODEditorSubsystem::Initialize(FSubsystemCollectionBase& Collection)
 {
-	// Ensure the WorldPartitionSubsystem gets created before the HLODEditorSubsystem
-	Collection.InitializeDependency<UWorldPartitionSubsystem>();
-
 	Super::Initialize(Collection);
 
 	bForceHLODStateUpdate = true;
@@ -157,9 +157,9 @@ void UWorldPartitionHLODEditorSubsystem::Initialize(FSubsystemCollectionBase& Co
 	CachedHLODMinDrawDistance = 0;
 	CachedHLODMaxDrawDistance = 0;
 	bCachedShowHLODsOverLoadedRegions = false;
+	bHLODSettingsFilteringActive = false;
 	
-	GetWorld()->OnWorldPartitionInitialized().AddUObject(this, &UWorldPartitionHLODEditorSubsystem::OnWorldPartitionInitialized);
-	GetWorld()->OnWorldPartitionUninitialized().AddUObject(this, &UWorldPartitionHLODEditorSubsystem::OnWorldPartitionUninitialized);
+	FWorldDelegates::OnPostWorldInitialization.AddUObject(this, &UWorldPartitionHLODEditorSubsystem::OnPostWorldInitialization);
 
 	GEngine->OnLevelActorListChanged().AddUObject(this, &UWorldPartitionHLODEditorSubsystem::ForceHLODStateUpdate);
 
@@ -171,6 +171,16 @@ void UWorldPartitionHLODEditorSubsystem::Initialize(FSubsystemCollectionBase& Co
 
 	HLODActorEditorRegisteredHandle = AWorldPartitionHLOD::GetHLODActorEditorRegisteredDelegate().AddUObject(this, &UWorldPartitionHLODEditorSubsystem::OnHLODActorRegistered);
 	HLODActorEditorUnregisteredHandle = AWorldPartitionHLOD::GetHLODActorEditorUnregisteredDelegate().AddUObject(this, &UWorldPartitionHLODEditorSubsystem::OnHLODActorUnregistered);
+
+	FLevelEditorModule& LevelEditorModule = FModuleManager::LoadModuleChecked<FLevelEditorModule>("LevelEditor");
+	if (TSharedPtr<ILevelEditor> FirstLevelEditor = LevelEditorModule.GetFirstLevelEditor())
+	{
+		OnLevelEditorCreated(FirstLevelEditor);
+	}
+	else
+	{
+		LevelEditorModule.OnLevelEditorCreated().AddUObject(this, &UWorldPartitionHLODEditorSubsystem::OnLevelEditorCreated);
+	}
 }
 
 void UWorldPartitionHLODEditorSubsystem::Deinitialize()
@@ -186,8 +196,325 @@ void UWorldPartitionHLODEditorSubsystem::Deinitialize()
 
 	GEngine->OnLevelActorListChanged().RemoveAll(this);
 
-	GetWorld()->OnWorldPartitionInitialized().RemoveAll(this);
-	GetWorld()->OnWorldPartitionUninitialized().RemoveAll(this);
+	FLevelEditorModule& LevelEditorModule = FModuleManager::LoadModuleChecked<FLevelEditorModule>("LevelEditor");
+	LevelEditorModule.OnLevelEditorCreated().RemoveAll(this);
+
+	FWorldDelegates::OnPostWorldInitialization.RemoveAll(this);
+}
+
+void UWorldPartitionHLODEditorSubsystem::OnLevelEditorCreated(TSharedPtr<ILevelEditor> InLevelEditor)
+{
+	if (GEditor)
+	{
+		ExtendShowFlagsMenu();
+	}
+}
+
+void UWorldPartitionHLODEditorSubsystem::ExtendShowFlagsMenu()
+{
+	check(GEditor);
+
+	UToolMenu* const Menu = UToolMenus::Get()->ExtendMenu("LevelEditor.ViewportToolbar.Show");
+	if (!Menu)
+	{
+		return;
+	}
+
+	FToolMenuSection& AllShowFlagsSection = Menu->FindOrAddSection("AllShowFlags");
+
+	// This is a dynamic entry so we can skip adding the submenu if the context
+	// indicates that the viewport's world isn't partitioned.
+	AllShowFlagsSection.AddEntry(FToolMenuEntry::InitDynamicEntry(
+		"ShowHLODsDynamic",
+		FNewToolMenuSectionDelegate::CreateLambda(
+			[](FToolMenuSection& InDynamicSection) -> void
+			{
+				TSharedPtr<SLevelViewport> Viewport = ULevelViewportContext::GetLevelViewport(InDynamicSection);
+				if (!Viewport)
+				{
+					return;
+				}
+
+				UWorld* World = Viewport->GetWorld();
+				if (!World)
+				{
+					return;
+				}
+
+				// Only add this submenu for partitioned worlds.
+				if (!World->IsPartitionedWorld())
+				{
+					return;
+				}
+
+				InDynamicSection.AddSubMenu(
+					"ShowHLODsMenu",
+					LOCTEXT("ShowHLODsMenu", "HLODs"),
+					LOCTEXT("ShowHLODsMenu_ToolTip", "Settings for HLODs in editor"),
+					FNewToolMenuDelegate::CreateLambda(
+						[](UToolMenu* Submenu)
+						{
+							TSharedPtr<SLevelViewport> Viewport = ULevelViewportContext::GetLevelViewport(Submenu);
+							if (!Viewport)
+							{
+								return;
+							}
+
+							UWorld* World = Viewport->GetWorld();
+							UWorldPartition* WorldPartition = World ? World->GetWorldPartition() : nullptr;
+							if (!WorldPartition)
+							{
+								return;
+							}
+
+							IWorldPartitionEditorModule* WorldPartitionEditorModule =
+								FModuleManager::GetModulePtr<IWorldPartitionEditorModule>("WorldPartitionEditor");
+							if (!WorldPartitionEditorModule)
+							{
+								return;
+							}
+
+							FText HLODInEditorDisallowedReason;
+							const bool bHLODInEditorAllowed =
+								WorldPartitionEditorModule->IsHLODInEditorAllowed(World, &HLODInEditorDisallowedReason);
+
+							// Show HLODs
+							{
+								FToolUIAction UIAction;
+								UIAction.ExecuteAction = FToolMenuExecuteAction::CreateLambda(
+									[WorldPartitionEditorModule](const FToolMenuContext& InContext)
+									{
+										WorldPartitionEditorModule->SetShowHLODsInEditor(
+											!WorldPartitionEditorModule->GetShowHLODsInEditor()
+										);
+									}
+								);
+								UIAction.CanExecuteAction = FToolMenuCanExecuteAction::CreateLambda(
+									[bHLODInEditorAllowed](const FToolMenuContext& InContext)
+									{
+										return bHLODInEditorAllowed;
+									}
+								);
+								UIAction.GetActionCheckState = FToolMenuGetActionCheckState::CreateLambda(
+									[WorldPartitionEditorModule](const FToolMenuContext& InContext)
+									{
+										return WorldPartitionEditorModule->GetShowHLODsInEditor()
+												 ? ECheckBoxState::Checked
+												 : ECheckBoxState::Unchecked;
+									}
+								);
+								FToolMenuEntry MenuEntry = FToolMenuEntry::InitMenuEntry(
+									"ShowHLODs",
+									LOCTEXT("ShowHLODs", "Show HLODs"),
+									bHLODInEditorAllowed ? LOCTEXT("ShowHLODsToolTip", "Show/Hide HLODs")
+														 : HLODInEditorDisallowedReason,
+									FSlateIcon(),
+									UIAction,
+									EUserInterfaceActionType::ToggleButton
+								);
+								Submenu->AddMenuEntry(NAME_None, MenuEntry);
+							}
+
+							// Show HLODs Over Loaded Regions
+							{
+								FToolUIAction UIAction;
+								UIAction.ExecuteAction = FToolMenuExecuteAction::CreateLambda(
+									[WorldPartitionEditorModule](const FToolMenuContext& InContext)
+									{
+										WorldPartitionEditorModule->SetShowHLODsOverLoadedRegions(
+											!WorldPartitionEditorModule->GetShowHLODsOverLoadedRegions()
+										);
+									}
+								);
+								UIAction.CanExecuteAction = FToolMenuCanExecuteAction::CreateLambda(
+									[bHLODInEditorAllowed](const FToolMenuContext& InContext)
+									{
+										return bHLODInEditorAllowed;
+									}
+								);
+								UIAction.GetActionCheckState = FToolMenuGetActionCheckState::CreateLambda(
+									[WorldPartitionEditorModule](const FToolMenuContext& InContext)
+									{
+										return WorldPartitionEditorModule->GetShowHLODsOverLoadedRegions()
+												 ? ECheckBoxState::Checked
+												 : ECheckBoxState::Unchecked;
+									}
+								);
+								FToolMenuEntry ShowHLODsEntry = FToolMenuEntry::InitMenuEntry(
+									"ShowHLODsOverLoadedRegions",
+									LOCTEXT("ShowHLODsOverLoadedRegions", "Show HLODs Over Loaded Regions"),
+									bHLODInEditorAllowed
+										? LOCTEXT("ShowHLODsOverLoadedRegions_ToolTip", "Show/Hide HLODs over loaded actors or regions")
+										: HLODInEditorDisallowedReason,
+									FSlateIcon(),
+									UIAction,
+									EUserInterfaceActionType::ToggleButton
+								);
+								Submenu->AddMenuEntry(NAME_None, ShowHLODsEntry);
+							}
+
+							// Min/Max Draw Distance
+							{
+								const double MinDrawDistanceMinValue = 0;
+								const double MinDrawDistanceMaxValue = 102400;
+
+								const double MaxDrawDistanceMinValue = 0;
+								const double MaxDrawDistanceMaxValue = 1638400;
+
+								// double SLevelViewportToolBar::OnGetHLODInEditorMinDrawDistanceValue() const
+								auto OnGetHLODInEditorMinDrawDistanceValue = []() -> double
+								{
+									IWorldPartitionEditorModule* WorldPartitionEditorModule =
+										FModuleManager::GetModulePtr<IWorldPartitionEditorModule>("WorldPartitionEditor");
+									return WorldPartitionEditorModule
+											 ? WorldPartitionEditorModule->GetHLODInEditorMinDrawDistance()
+											 : 0;
+								};
+
+								// void SLevelViewportToolBar::OnHLODInEditorMinDrawDistanceValueChanged(double NewValue) const
+								auto OnHLODInEditorMinDrawDistanceValueChanged = [](double NewValue) -> void
+								{
+									IWorldPartitionEditorModule* WorldPartitionEditorModule =
+										FModuleManager::GetModulePtr<IWorldPartitionEditorModule>("WorldPartitionEditor");
+									if (WorldPartitionEditorModule)
+									{
+										WorldPartitionEditorModule->SetHLODInEditorMinDrawDistance(NewValue);
+										GEditor->RedrawLevelEditingViewports(true);
+									}
+								};
+
+								TSharedRef<SSpinBox<double>> MinDrawDistanceSpinBox =
+									SNew(SSpinBox<double>)
+										.MinValue(MinDrawDistanceMinValue)
+										.MaxValue(MinDrawDistanceMaxValue)
+										.IsEnabled(bHLODInEditorAllowed)
+										.Value_Lambda(OnGetHLODInEditorMinDrawDistanceValue)
+										.OnValueChanged_Lambda(OnHLODInEditorMinDrawDistanceValueChanged)
+										.ToolTipText(
+											bHLODInEditorAllowed
+												? LOCTEXT(
+													  "HLODsInEditor_MinDrawDistance_Tooltip",
+													  "Sets the minimum distance at which HLOD will be rendered"
+												  )
+												: HLODInEditorDisallowedReason
+										)
+										.OnBeginSliderMovement_Lambda(
+											[]()
+											{
+												// Disable Slate throttling during slider drag to ensure immediate updates while moving the slider.
+												FSlateThrottleManager::Get().DisableThrottle(true);
+											}
+										)
+										.OnEndSliderMovement_Lambda(
+											[](double)
+											{
+												FSlateThrottleManager::Get().DisableThrottle(false);
+											}
+										);
+
+								// double SLevelViewportToolBar::OnGetHLODInEditorMaxDrawDistanceValue() const
+								auto OnGetHLODInEditorMaxDrawDistanceValue = []() -> double
+								{
+									IWorldPartitionEditorModule* WorldPartitionEditorModule =
+										FModuleManager::GetModulePtr<IWorldPartitionEditorModule>("WorldPartitionEditor");
+									return WorldPartitionEditorModule
+											 ? WorldPartitionEditorModule->GetHLODInEditorMaxDrawDistance()
+											 : 0;
+								};
+
+								// void SLevelViewportToolBar::OnHLODInEditorMaxDrawDistanceValueChanged(double NewValue) const
+								auto OnHLODInEditorMaxDrawDistanceValueChanged = [](double NewValue) -> void
+								{
+									IWorldPartitionEditorModule* WorldPartitionEditorModule =
+										FModuleManager::GetModulePtr<IWorldPartitionEditorModule>("WorldPartitionEditor");
+									if (WorldPartitionEditorModule)
+									{
+										WorldPartitionEditorModule->SetHLODInEditorMaxDrawDistance(NewValue);
+										GEditor->RedrawLevelEditingViewports(true);
+									}
+								};
+
+								TSharedRef<SSpinBox<double>> MaxDrawDistanceSpinBox =
+									SNew(SSpinBox<double>)
+										.MinValue(MaxDrawDistanceMinValue)
+										.MaxValue(MaxDrawDistanceMaxValue)
+										.IsEnabled(bHLODInEditorAllowed)
+										.Value_Lambda(OnGetHLODInEditorMaxDrawDistanceValue)
+										.OnValueChanged_Lambda(OnHLODInEditorMaxDrawDistanceValueChanged)
+										.ToolTipText(
+											bHLODInEditorAllowed
+												? LOCTEXT(
+													  "HLODsInEditor_MaxDrawDistance_Tooltip", "Sets the maximum distance at which HLODs will be rendered (0.0 means infinite)"
+												  )
+												: HLODInEditorDisallowedReason
+										)
+										.OnBeginSliderMovement_Lambda(
+											[]()
+											{
+												// Disable Slate throttling during slider drag to ensure immediate updates while moving the slider.
+												FSlateThrottleManager::Get().DisableThrottle(true);
+											}
+										)
+										.OnEndSliderMovement_Lambda(
+											[](double)
+											{
+												FSlateThrottleManager::Get().DisableThrottle(false);
+											}
+										);
+
+								auto CreateDrawDistanceWidget = [](TSharedRef<SSpinBox<double>> InSpinBoxWidget)
+								{
+									// clang-format off
+									return SNew(SBox)
+										.HAlign(HAlign_Right)
+										[
+											SNew(SBox)
+										  .Padding(FMargin(0.0f, 0.0f, 0.0f, 0.0f))
+										  .WidthOverride(100.0f)
+											[
+												SNew(SBorder)
+												.BorderImage(FAppStyle::Get().GetBrush("Menu.WidgetBorder"))
+												.Padding(FMargin(1.0f))
+												[
+													InSpinBoxWidget
+												]
+											]
+										];
+									// clang-format on
+								};
+
+								FToolMenuEntry MinDrawDistanceMenuEntry = FToolMenuEntry::InitWidget(
+									"Min Draw Distance",
+									CreateDrawDistanceWidget(MinDrawDistanceSpinBox),
+									LOCTEXT("MinDrawDistance", "Min Draw Distance")
+								);
+								Submenu->AddMenuEntry(NAME_None, MinDrawDistanceMenuEntry);
+
+								FToolMenuEntry MaxDrawDistanceMenuEntry = FToolMenuEntry::InitWidget(
+									"Max Draw Distance",
+									CreateDrawDistanceWidget(MaxDrawDistanceSpinBox),
+									LOCTEXT("MaxDrawDistance", "Max Draw Distance")
+								);
+								Submenu->AddMenuEntry(NAME_None, MaxDrawDistanceMenuEntry);
+							}
+						}
+					),
+					false,
+					FSlateIcon(FAppStyle::Get().GetStyleSetName(), "ShowFlagsMenu.SubMenu.HLODs")
+				);
+			}
+		)
+	));
+}
+
+void UWorldPartitionHLODEditorSubsystem::OnPostWorldInitialization(UWorld* InWorld, const UWorld::InitializationValues InIVS)
+{
+	if (InWorld && InWorld->WorldType == EWorldType::Editor)
+	{
+		check(!WorldPartitionsHLODEditorData.Contains(InWorld->GetWorldPartition()));
+		InWorld->OnWorldPartitionInitialized().AddUObject(this, &UWorldPartitionHLODEditorSubsystem::OnWorldPartitionInitialized);
+		InWorld->OnWorldPartitionUninitialized().AddUObject(this, &UWorldPartitionHLODEditorSubsystem::OnWorldPartitionUninitialized);
+	}
 }
 
 void UWorldPartitionHLODEditorSubsystem::OnWorldPartitionEditorSettingsChanged(const FName& InPropertyName, const UWorldPartitionEditorSettings& InWorldPartitionEditorSettings)
@@ -202,14 +529,20 @@ void UWorldPartitionHLODEditorSubsystem::ApplyHLODSettingsFiltering()
 {
 	static const FName PropertyPermissionListOwnerName = "AdvancedHLODSettingsFiltering";
 
-	FPropertyEditorPermissionList::Get().UnregisterOwner(PropertyPermissionListOwnerName);
-
-	UEnum* HLODLayerTypeEnum = StaticEnum<EHLODLayerType>();
-	if (!GetDefault<UWorldPartitionEditorSettings>()->GetEnableAdvancedHLODSettings())
+	if (bHLODSettingsFilteringActive != !GetDefault<UWorldPartitionEditorSettings>()->GetEnableAdvancedHLODSettings())
 	{
-		for (const TPair<TSoftObjectPtr<UStruct>, TSet<FName>>& StructProperties : StructsPropertiesVisibility.FindOrAdd(EHLODSettingsVisibility::BasicSettings))
+		bHLODSettingsFilteringActive = !GetDefault<UWorldPartitionEditorSettings>()->GetEnableAdvancedHLODSettings();
+
+		if (!bHLODSettingsFilteringActive)
 		{
-			FPropertyEditorPermissionList::Get().AddToAllowList(StructProperties.Key, StructProperties.Value.Array(), PropertyPermissionListOwnerName);
+			FPropertyEditorPermissionList::Get().UnregisterOwner(PropertyPermissionListOwnerName);
+		}
+		else
+		{
+			for (const TPair<TSoftObjectPtr<UStruct>, TSet<FName>>& StructProperties : StructsPropertiesVisibility.FindOrAdd(EHLODSettingsVisibility::BasicSettings))
+			{
+				FPropertyEditorPermissionList::Get().AddToAllowList(StructProperties.Key, StructProperties.Value.Array(), PropertyPermissionListOwnerName);
+			}
 		}
 	}
 }
@@ -289,6 +622,11 @@ void UWorldPartitionHLODEditorSubsystem::ForceHLODStateUpdate()
 	}
 }
 
+bool UWorldPartitionHLODEditorSubsystem::IsTickable() const
+{
+	return GEditor != nullptr;
+}
+
 void UWorldPartitionHLODEditorSubsystem::Tick(float DeltaTime)
 {
 	TRACE_CPUPROFILER_EVENT_SCOPE(UWorldPartitionHLODEditorSubsystem::Tick);
@@ -380,20 +718,18 @@ void UWorldPartitionHLODEditorSubsystem::AddHLODSettingsFilter(EHLODSettingsVisi
 	StructsPropertiesVisibility.FindOrAdd(InSettingsVisibility).FindOrAdd(InStruct).Add(InPropertyName);
 }
 
-bool UWorldPartitionHLODEditorSubsystem::WriteHLODStats(const IWorldPartitionEditorModule::FWriteHLODStatsParams& Params) const
+bool UWorldPartitionHLODEditorSubsystem::WriteHLODStats(const IWorldPartitionEditorModule::FWriteHLODStatsParams& Params)
 {
-	check(Params.World == GetWorld());
-
 	bool bResult = false;
 
 	switch(Params.StatsType)
 	{
 	case IWorldPartitionEditorModule::FWriteHLODStatsParams::EStatsType::Default:
-		bResult = WriteHLODStats(Params.Filename);
+		bResult = WriteHLODStats(Params.World, Params.Filename);
 		break;
 
 	case IWorldPartitionEditorModule::FWriteHLODStatsParams::EStatsType::InputDetails:
-		bResult = WriteHLODInputStats(Params.Filename);
+		bResult = WriteHLODInputStats(Params.World, Params.Filename);
 		break;
 	}
 
@@ -409,10 +745,9 @@ bool UWorldPartitionHLODEditorSubsystem::WriteHLODStats(const IWorldPartitionEdi
 	return bResult;
 }
 
-bool UWorldPartitionHLODEditorSubsystem::WriteHLODStats(const FString& InFilename) const
+bool UWorldPartitionHLODEditorSubsystem::WriteHLODStats(UWorld* InWorld, const FString& InFilename)
 {
-	UWorld* World = GetWorld();
-	UWorldPartition* WorldPartition = World->GetWorldPartition();
+	UWorldPartition* WorldPartition = InWorld->GetWorldPartition();
 	if (!WorldPartition)
 	{
 		return false;
@@ -438,7 +773,7 @@ bool UWorldPartitionHLODEditorSubsystem::WriteHLODStats(const FString& InFilenam
 
 	TArray<TPair<FName, FGetStatFunc>> StatsToWrite =
 	{
-		{ "WorldPackage",		[World](FWorldPartitionActorDescInstance* InActorDescInstance, const FHLODActorDesc& InActorDesc) { return World->GetPackage()->GetName(); } },
+		{ "WorldPackage",		[InWorld](FWorldPartitionActorDescInstance* InActorDescInstance, const FHLODActorDesc& InActorDesc) { return InWorld->GetPackage()->GetName(); } },
 		{ "Name",				[](FWorldPartitionActorDescInstance* InActorDescInstance, const FHLODActorDesc& InActorDesc) { return InActorDescInstance->GetActorLabelString(); } },
 		{ "HLODLayer",			[](FWorldPartitionActorDescInstance* InActorDescInstance, const FHLODActorDesc& InActorDesc) { return InActorDesc.GetSourceHLODLayer().GetAssetName().ToString(); }},
 		{ "SpatiallyLoaded",	[](FWorldPartitionActorDescInstance* InActorDescInstance, const FHLODActorDesc& InActorDesc) { return InActorDescInstance->GetIsSpatiallyLoaded() ? TEXT("true") : TEXT("false"); } },
@@ -495,10 +830,9 @@ bool UWorldPartitionHLODEditorSubsystem::WriteHLODStats(const FString& InFilenam
 	return FFileHelper::SaveStringToFile(Output, *InFilename, FFileHelper::EEncodingOptions::AutoDetect, &IFileManager::Get(), EFileWrite::FILEWRITE_Append);
 }
 
-bool UWorldPartitionHLODEditorSubsystem::WriteHLODInputStats(const FString& InFilename) const
+bool UWorldPartitionHLODEditorSubsystem::WriteHLODInputStats(UWorld* InWorld, const FString& InFilename)
 {
-	UWorld* World = GetWorld();
-	UWorldPartition* WorldPartition = World->GetWorldPartition();
+	UWorldPartition* WorldPartition = InWorld->GetWorldPartition();
 	if (!WorldPartition)
 	{
 		return false;
@@ -510,7 +844,7 @@ bool UWorldPartitionHLODEditorSubsystem::WriteHLODInputStats(const FString& InFi
 	TMap<TPair<int32, FName>, FHLODBuildInputReferencedAssets> BuildersReferencedAssets;
 
 	// Aggregate referenced assets from all HLOD actors
-	FWorldPartitionHelpers::ForEachActorWithLoading(WorldPartition, [this, World, WorldPartition, &BuildersReferencedAssets](const FWorldPartitionActorDescInstance* ActorDescInstance)
+	FWorldPartitionHelpers::ForEachActorWithLoading(WorldPartition, [WorldPartition, &BuildersReferencedAssets](const FWorldPartitionActorDescInstance* ActorDescInstance)
 	{
 		AWorldPartitionHLOD* HLODActor = Cast<AWorldPartitionHLOD>(ActorDescInstance->GetActor());
 		if (!HLODActor)
@@ -586,14 +920,11 @@ FAutoConsoleCommand HLODDumpStats(
 		{
 			if (UWorld* World = Context.World())
 			{
-				if (UWorldPartitionHLODEditorSubsystem* HLODEditorSubsystem = World->GetSubsystem<UWorldPartitionHLODEditorSubsystem>())
-				{
-					IWorldPartitionEditorModule::FWriteHLODStatsParams Params;
-					Params.World = World;
-					Params.StatsType = IWorldPartitionEditorModule::FWriteHLODStatsParams::EStatsType::Default;
-					Params.Filename = HLODStatsOutputFilename;
-					HLODEditorSubsystem->WriteHLODStats(Params);
-				}
+				IWorldPartitionEditorModule::FWriteHLODStatsParams Params;
+				Params.World = World;
+				Params.StatsType = IWorldPartitionEditorModule::FWriteHLODStatsParams::EStatsType::Default;
+				Params.Filename = HLODStatsOutputFilename;
+				UWorldPartitionHLODEditorSubsystem::WriteHLODStats(Params);
 			}
 		}
 	})
@@ -610,14 +941,11 @@ FAutoConsoleCommand HLODDumpInputStats(
 		{
 			if (UWorld* World = Context.World())
 			{
-				if (UWorldPartitionHLODEditorSubsystem* HLODEditorSubsystem = World->GetSubsystem<UWorldPartitionHLODEditorSubsystem>())
-				{
-					IWorldPartitionEditorModule::FWriteHLODStatsParams Params;
-					Params.World = World;
-					Params.StatsType = IWorldPartitionEditorModule::FWriteHLODStatsParams::EStatsType::InputDetails;
-					Params.Filename = HLODStatsOutputFilename;
-					HLODEditorSubsystem->WriteHLODStats(Params);
-				}
+				IWorldPartitionEditorModule::FWriteHLODStatsParams Params;
+				Params.World = World;
+				Params.StatsType = IWorldPartitionEditorModule::FWriteHLODStatsParams::EStatsType::InputDetails;
+				Params.Filename = HLODStatsOutputFilename;
+				UWorldPartitionHLODEditorSubsystem::WriteHLODStats(Params);
 			}
 		}
 	})
diff --git a/Engine/Source/Editor/WorldPartitionEditor/Private/WorldPartitionEditorModule.cpp b/Engine/Source/Editor/WorldPartitionEditor/Private/WorldPartitionEditorModule.cpp
--- a/Engine/Source/Editor/WorldPartitionEditor/Private/WorldPartitionEditorModule.cpp
+++ b/Engine/Source/Editor/WorldPartitionEditor/Private/WorldPartitionEditorModule.cpp
@@ -610,12 +610,7 @@ bool FWorldPartitionEditorModule::IsHLODInEditorAllowed(UWorld* InWorld, FText*
 
 bool FWorldPartitionEditorModule::WriteHLODStats(const FWriteHLODStatsParams& Params) const
 {
-	UWorldPartitionHLODEditorSubsystem* HLODEditorSubsystem = Params.World->GetSubsystem<UWorldPartitionHLODEditorSubsystem>();
-	if (ensure(HLODEditorSubsystem))
-	{
-		return HLODEditorSubsystem->WriteHLODStats(Params);
-	}
-	return false;
+	return UWorldPartitionHLODEditorSubsystem::WriteHLODStats(Params);
 }
 
 void FWorldPartitionEditorModule::OnConvertMap()
diff --git a/Engine/Source/Editor/WorldPartitionEditor/Public/WorldPartition/HLOD/HLODEditorSubsystem.h b/Engine/Source/Editor/WorldPartitionEditor/Public/WorldPartition/HLOD/HLODEditorSubsystem.h
--- a/Engine/Source/Editor/WorldPartitionEditor/Public/WorldPartition/HLOD/HLODEditorSubsystem.h
+++ b/Engine/Source/Editor/WorldPartitionEditor/Public/WorldPartition/HLOD/HLODEditorSubsystem.h
@@ -2,12 +2,13 @@
 
 #pragma once
 
-#include "Subsystems/WorldSubsystem.h"
-#include "WorldPartition/IWorldPartitionEditorModule.h"
-#include "WorldPartition/WorldPartitionActorLoaderInterface.h"
-#include "WorldPartition/WorldPartitionHandle.h"
+#include "EditorSubsystem.h"
 #include "Templates/PimplPtr.h"
+#include "Tickable.h"
 #include "UObject/ObjectKey.h"
+#include "Engine/World.h"
+#include "WorldPartition/IWorldPartitionEditorModule.h"
+#include "WorldPartition/WorldPartitionActorLoaderInterface.h"
 
 #include "HLODEditorSubsystem.generated.h"
 
@@ -16,6 +17,7 @@
 
 class AActor;
 class AWorldPartitionHLOD;
+class ILevelEditor;
 class UPrimitiveComponent;
 class UWorldPartition;
 class UWorldPartitionEditorSettings;
@@ -34,7 +36,7 @@ enum class EHLODSettingsVisibility : uint8
  * UWorldPartitionHLODEditorSubsystem
  */
 UCLASS(MinimalAPI)
-class UWorldPartitionHLODEditorSubsystem : public UTickableWorldSubsystem
+class UWorldPartitionHLODEditorSubsystem : public UEditorSubsystem, public FTickableGameObject
 {
 	GENERATED_BODY()
 
@@ -47,18 +49,14 @@ public:
 	UE_API virtual void Deinitialize() override;
 	//~ End USubsystem Interface.
 
-	//~ Begin UWorldSubsystem Interface.
-	UE_API virtual bool DoesSupportWorldType(const EWorldType::Type WorldType) const override;
-	//~ End UWorldSubsystem Interface.
-
 	//~ Begin FTickableGameObject Interface
-	virtual bool IsTickable() const override { return true; }
+	virtual bool IsTickable() const override;
 	virtual bool IsTickableInEditor() const override { return true; }
 	UE_API virtual void Tick(float DeltaTime) override;
 	UE_API virtual TStatId GetStatId() const override;
 	//~ End FTickableGameObject Interface
 
-	UE_API virtual bool WriteHLODStats(const IWorldPartitionEditorModule::FWriteHLODStatsParams& Params) const;
+	static UE_API bool WriteHLODStats(const IWorldPartitionEditorModule::FWriteHLODStatsParams& Params);
 
 	static UE_API void AddHLODSettingsFilter(EHLODSettingsVisibility InSettingsVisibility, TSoftObjectPtr<UStruct> InStruct, FName InPropertyName);
 
@@ -75,8 +73,8 @@ private:
 
 	UE_API void ForceHLODStateUpdate();
 
-	UE_API bool WriteHLODStats(const FString& InFilename) const;
-	UE_API bool WriteHLODInputStats(const FString& InFilename) const;
+	static UE_API bool WriteHLODStats(UWorld* InWorld, const FString& InFilename);
+	static UE_API bool WriteHLODInputStats(UWorld* InWorld, const FString& InFilename);
 	
 	UE_API void OnWorldPartitionEditorSettingsChanged(const FName& PropertyName, const UWorldPartitionEditorSettings& WorldPartitionEditorSettings);
 	UE_API void ApplyHLODSettingsFiltering();
@@ -86,6 +84,11 @@ private:
 	void OnHLODActorRegistered(AWorldPartitionHLOD* Actor);
 	void OnHLODActorUnregistered(AWorldPartitionHLOD* Actor);
 
+	void OnLevelEditorCreated(TSharedPtr<ILevelEditor> InLevelEditor);
+	void ExtendShowFlagsMenu();
+
+	void OnPostWorldInitialization(UWorld* InWorld, const UWorld::InitializationValues InIVS);
+
 private:
 	FDelegateHandle HLODActorEditorRegisteredHandle;
 	FDelegateHandle HLODActorEditorUnregisteredHandle;
@@ -96,6 +99,7 @@ private:
 	bool bCachedShowHLODsOverLoadedRegions;
 	bool bForceHLODStateUpdate;
 	bool bForceHLODVisibilityUpdate;
+	bool bHLODSettingsFilteringActive;
 
 	TMap<TObjectKey<UWorldPartition>, TPimplPtr<FWorldPartitionHLODEditorData>> WorldPartitionsHLODEditorData;
 
