World Partition - HLOD: Fix use-after-free crash in Export HLOD Assets dialog

The expectation that providing a stack-local pointer to the window was safe, as it was modal. But in practice it's possible for a refresh to occur after the EditExportHLODAssetsParams() function has exited.

#Synced-CL 51374974

#Synced-CL 51374974

diff --git a/Engine/Source/Editor/WorldPartitionEditor/Private/WorldPartition/Customizations/WorldPartitionHLODDetailsCustomization.cpp b/Engine/Source/Editor/WorldPartitionEditor/Private/WorldPartition/Customizations/WorldPartitionHLODDetailsCustomization.cpp
--- a/Engine/Source/Editor/WorldPartitionEditor/Private/WorldPartition/Customizations/WorldPartitionHLODDetailsCustomization.cpp
+++ b/Engine/Source/Editor/WorldPartitionEditor/Private/WorldPartition/Customizations/WorldPartitionHLODDetailsCustomization.cpp
@@ -230,9 +230,6 @@ bool FWorldPartitionHLODDetailsCustomization::CanExportAssets() const
 
 static EAppReturnType::Type EditExportHLODAssetsParams(FExportHLODAssetsParams& InOutParams, const TArray<AWorldPartitionHLOD*>& InHLODActorsToExport)
 {
-	// Work on a copy; commit only if user accepts.
-	FExportHLODAssetsParams WorkingCopy = InOutParams;
-
 	// Setup property editor for a struct
 	FPropertyEditorModule& PropertyEditorModule = FModuleManager::LoadModuleChecked<FPropertyEditorModule>("PropertyEditor");
 
@@ -249,10 +246,9 @@ static EAppReturnType::Type EditExportHLODAssetsParams(FExportHLODAssetsParams&
 	TSharedRef<IStructureDetailsView> StructDetails =
 		PropertyEditorModule.CreateStructureDetailView(DetailsArgs, StructArgs, nullptr);
 
-	TSharedRef<FStructOnScope> StructData = MakeShared<FStructOnScope>(
-		FExportHLODAssetsParams::StaticStruct(),
-		reinterpret_cast<uint8*>(&WorkingCopy) // valid: modal window blocks until closed
-	);
+	TSharedRef<FStructOnScope> StructData = MakeShared<FStructOnScope>(FExportHLODAssetsParams::StaticStruct());
+	FExportHLODAssetsParams& WorkingCopy = *reinterpret_cast<FExportHLODAssetsParams*>(StructData->GetStructMemory());
+	FExportHLODAssetsParams::StaticStruct()->CopyScriptStruct(&WorkingCopy, &InOutParams);
 	StructDetails->SetStructureData(StructData);
 	
 	auto CanExport = [&WorkingCopy, &InHLODActorsToExport](FText& OutExportErrorMessage) -> bool
