[HLOD] Support 'Export HLOD Assets' for HLOD Actors with externalized resource packages.

#Synced-CL 51416907

diff --git a/Engine/Source/Editor/WorldPartitionEditor/Private/WorldPartition/Customizations/WorldPartitionHLODDetailsCustomization.cpp b/Engine/Source/Editor/WorldPartitionEditor/Private/WorldPartition/Customizations/WorldPartitionHLODDetailsCustomization.cpp
--- a/Engine/Source/Editor/WorldPartitionEditor/Private/WorldPartition/Customizations/WorldPartitionHLODDetailsCustomization.cpp
+++ b/Engine/Source/Editor/WorldPartitionEditor/Private/WorldPartition/Customizations/WorldPartitionHLODDetailsCustomization.cpp
@@ -208,7 +208,7 @@ bool FWorldPartitionHLODDetailsCustomization::CanExportHLODAssets(const AWorldPa
 	}
 
 	bool bHasPrivateStaticMesh = false;
-	ForEachObjectWithPackage(HLODActor->GetPackage(), [&bHasPrivateStaticMesh](UObject* InObject)
+	ForEachObjectWithPackage(HLODActor->GetHLODResourcesPackage(), [&bHasPrivateStaticMesh](UObject* InObject)
 	{
 		bHasPrivateStaticMesh = InObject->IsA<UStaticMesh>();
 		return !bHasPrivateStaticMesh; // Continue until we find a static mesh
diff --git a/Engine/Source/Runtime/Engine/Private/WorldPartition/HLOD/HLODActor.cpp b/Engine/Source/Runtime/Engine/Private/WorldPartition/HLOD/HLODActor.cpp
--- a/Engine/Source/Runtime/Engine/Private/WorldPartition/HLOD/HLODActor.cpp
+++ b/Engine/Source/Runtime/Engine/Private/WorldPartition/HLOD/HLODActor.cpp
@@ -858,10 +858,10 @@ TArray<UObject*> AWorldPartitionHLOD::ExportHLODAssets(const FExportHLODAssetsPa
 	TArray<UObject*> ExportedAssets;
 	OutErrorMessage = TEXT("");
 		
-	auto ShouldExportAsset = [ActorPackage = GetPackage()](const UObject* InAsset)
+	auto ShouldExportAsset = [ResourcesPackage = GetHLODResourcesPackage()](const UObject* InAsset)
 	{
 		// Skip null entries and assets that are not in the HLOD package (no need for export)
-		return InAsset && InAsset->IsInPackage(ActorPackage);
+		return InAsset && InAsset->IsInPackage(ResourcesPackage);
 	};
 
 	// Unbuilt HLODs will have no static mesh components, we can ignore them
