[HLOD] Fix crash in CanExportHLODAssets that could occur when opening the Details panel for instanced HLOD actors with externalized HLOD assets enabled.

#Synced-CL 51416907

diff --git a/Engine/Source/Editor/WorldPartitionEditor/Private/WorldPartition/Customizations/WorldPartitionHLODDetailsCustomization.cpp b/Engine/Source/Editor/WorldPartitionEditor/Private/WorldPartition/Customizations/WorldPartitionHLODDetailsCustomization.cpp
--- a/Engine/Source/Editor/WorldPartitionEditor/Private/WorldPartition/Customizations/WorldPartitionHLODDetailsCustomization.cpp
+++ b/Engine/Source/Editor/WorldPartitionEditor/Private/WorldPartition/Customizations/WorldPartitionHLODDetailsCustomization.cpp
@@ -207,6 +207,11 @@ bool FWorldPartitionHLODDetailsCustomization::CanExportHLODAssets(const AWorldPa
 		return false;
 	}
 
+	if (HLODActor->GetHLODResourcesPackage() == nullptr)
+	{
+		return false;
+	}
+
 	bool bHasPrivateStaticMesh = false;
 	ForEachObjectWithPackage(HLODActor->GetHLODResourcesPackage(), [&bHasPrivateStaticMesh](UObject* InObject)
 	{
