World Partition - HLOD: Compare wipe overlay tool for comparing HLOD vs source actors

Adds an interactive wipe/overlay comparison window for HLOD actors.
A draggable vertical line separates source actors (left) from HLOD (right), both rendered at the same camera position with synced navigation.

- Wipe overlay layout using SOverlay + SHLODClipWrapper (PushClip/PopClip)
- Draggable wipe handle (SHLODWipeHandle) with resize cursor
- Source actor loading directly into preview scene world (no temp world)
- View mode toolbar with buffer visualization support
- Camera and view mode sync between both viewports
- Lighting duplicated from the main editor world
- Initial camera matches the active editor viewport
- Distance ruler to help understand at which distance the HLOD is usually seen
- Progress dialog while loading source actors
- "Compare vs Source" button in HLOD details panel
- "Compare HLOD vs Source" entry in viewport right-click context menu

#Synced-CL 51774560

diff --git a/Engine/Source/Editor/WorldPartitionEditor/Private/WorldPartition/Customizations/WorldPartitionHLODDetailsCustomization.cpp b/Engine/Source/Editor/WorldPartitionEditor/Private/WorldPartition/Customizations/WorldPartitionHLODDetailsCustomization.cpp
--- a/Engine/Source/Editor/WorldPartitionEditor/Private/WorldPartition/Customizations/WorldPartitionHLODDetailsCustomization.cpp
+++ b/Engine/Source/Editor/WorldPartitionEditor/Private/WorldPartition/Customizations/WorldPartitionHLODDetailsCustomization.cpp
@@ -30,6 +30,8 @@
 #include "Widgets/Notifications/SNotificationList.h"
 #include "WorldPartition/WorldPartition.h"
 #include "WorldPartition/HLOD/HLODActor.h"
+#include "WorldPartition/HLOD/SHLODCompareWindow.h"
+#include "Framework/Application/SlateApplication.h"
 
 #define LOCTEXT_NAMESPACE "FWorldPartitionHLODDetailsCustomization"
 
@@ -77,6 +79,16 @@ void FWorldPartitionHLODDetailsCustomization::CustomizeDetails(IDetailLayoutBuil
 			.OnClicked(this, &FWorldPartitionHLODDetailsCustomization::OnExportAssets)
 			.IsEnabled(this, &FWorldPartitionHLODDetailsCustomization::CanExportAssets)
 		]
+		+ SVerticalBox::Slot()
+		[
+			SNew(SButton)
+			.HAlign(EHorizontalAlignment::HAlign_Center)
+			.VAlign(EVerticalAlignment::VAlign_Center)
+			.Text(LOCTEXT("CompareHLODButtonText", "Compare Against Source"))
+			.ToolTipText(LOCTEXT("CompareHLODButtonTooltip", "Open a wipe comparison of the HLOD against its source actors"))
+			.OnClicked(this, &FWorldPartitionHLODDetailsCustomization::OnCompareHLOD)
+			.IsEnabled(this, &FWorldPartitionHLODDetailsCustomization::CanCompareHLOD)
+		]
 	];
 }
 
@@ -465,4 +477,75 @@ FReply FWorldPartitionHLODDetailsCustomization::OnExportAssets()
 
 	return FReply::Handled();
 }
+
+bool FWorldPartitionHLODDetailsCustomization::CanCompareHLOD() const
+{
+	TArray<AWorldPartitionHLOD*> SelectedHLODActors = GetSelectedHLODActors();
+	if (SelectedHLODActors.IsEmpty())
+	{
+		return false;
+	}
+
+	AWorldPartitionHLOD* HLODActor = SelectedHLODActors[0];
+	if (!HLODActor || HLODActor->IsTemplate())
+	{
+		return false;
+	}
+
+	if (HLODActor->GetSourceActors() == nullptr)
+	{
+		return false;
+	}
+
+	if (HLODActor->GetInstanceComponents().Num() == 0)
+	{
+		return false;
+	}
+
+	return true;
+}
+
+FReply FWorldPartitionHLODDetailsCustomization::OnCompareHLOD()
+{
+	TArray<AWorldPartitionHLOD*> SelectedHLODActors = GetSelectedHLODActors();
+	if (SelectedHLODActors.IsEmpty())
+	{
+		return FReply::Handled();
+	}
+
+	TArray<TWeakObjectPtr<AWorldPartitionHLOD>> WeakHLODActors;
+	for (AWorldPartitionHLOD* Actor : SelectedHLODActors)
+	{
+		if (Actor)
+		{
+			WeakHLODActors.Add(Actor);
+		}
+	}
+
+	if (WeakHLODActors.IsEmpty())
+	{
+		return FReply::Handled();
+	}
+
+	FText WindowTitle = (WeakHLODActors.Num() == 1)
+		? FText::Format(LOCTEXT("CompareWindowTitleSingle", "Compare HLOD: {0}"), FText::FromString(WeakHLODActors[0]->GetActorLabel()))
+		: FText::Format(LOCTEXT("CompareWindowTitleMulti", "Compare HLOD ({0} actors)"), FText::AsNumber(WeakHLODActors.Num()));
+
+	TSharedRef<SWindow> Window = SNew(SWindow)
+		.Title(WindowTitle)
+		.ClientSize(FVector2D(1280, 720))
+		.SizingRule(ESizingRule::UserSized)
+		.SupportsMinimize(true)
+		.SupportsMaximize(true);
+
+	Window->SetContent(
+		SNew(SHLODCompareWindow)
+		.HLODActors(WeakHLODActors)
+	);
+
+	FSlateApplication::Get().AddWindow(Window);
+
+	return FReply::Handled();
+}
+
 #undef LOCTEXT_NAMESPACE
diff --git a/Engine/Source/Editor/WorldPartitionEditor/Private/WorldPartition/Customizations/WorldPartitionHLODDetailsCustomization.h b/Engine/Source/Editor/WorldPartitionEditor/Private/WorldPartition/Customizations/WorldPartitionHLODDetailsCustomization.h
--- a/Engine/Source/Editor/WorldPartitionEditor/Private/WorldPartition/Customizations/WorldPartitionHLODDetailsCustomization.h
+++ b/Engine/Source/Editor/WorldPartitionEditor/Private/WorldPartition/Customizations/WorldPartitionHLODDetailsCustomization.h
@@ -34,6 +34,9 @@ private:
 	bool CanExportAssets() const;
 	FReply OnExportAssets();
 
+	bool CanCompareHLOD() const;
+	FReply OnCompareHLOD();
+
 private:
 	TArray<TWeakObjectPtr<UObject>> SelectedObjects;
 };
diff --git a/Engine/Source/Editor/WorldPartitionEditor/Private/WorldPartition/HLOD/HLODCompareCommands.cpp b/Engine/Source/Editor/WorldPartitionEditor/Private/WorldPartition/HLOD/HLODCompareCommands.cpp
new file mode 100644
--- /dev/null
+++ b/Engine/Source/Editor/WorldPartitionEditor/Private/WorldPartition/HLOD/HLODCompareCommands.cpp
@@ -0,0 +1,19 @@
+// Copyright Epic Games, Inc. All Rights Reserved.
+
+#include "WorldPartition/HLOD/HLODCompareCommands.h"
+
+#define LOCTEXT_NAMESPACE "HLODCompareCommands"
+
+void FHLODCompareCommands::RegisterCommands()
+{
+	UI_COMMAND(ViewModeLit, "Lit", "Switch to Lit view mode", EUserInterfaceActionType::RadioButton, FInputChord(EModifierKey::Alt, EKeys::One));
+	UI_COMMAND(ViewModeWireframe, "Wireframe", "Switch to Wireframe view mode", EUserInterfaceActionType::RadioButton, FInputChord(EModifierKey::Alt, EKeys::Two));
+	UI_COMMAND(ViewModeBaseColor, "Base Color", "Switch to Base Color buffer visualization", EUserInterfaceActionType::RadioButton, FInputChord(EModifierKey::Alt, EKeys::Three));
+	UI_COMMAND(ViewModeMetallic, "Metallic", "Switch to Metallic buffer visualization", EUserInterfaceActionType::RadioButton, FInputChord(EModifierKey::Alt, EKeys::Four));
+	UI_COMMAND(ViewModeRoughness, "Roughness", "Switch to Roughness buffer visualization", EUserInterfaceActionType::RadioButton, FInputChord(EModifierKey::Alt, EKeys::Five));
+	UI_COMMAND(ViewModeSpecular, "Specular", "Switch to Specular buffer visualization", EUserInterfaceActionType::RadioButton, FInputChord(EModifierKey::Alt, EKeys::Six));
+	UI_COMMAND(ViewModeWorldNormal, "World Normal", "Switch to World Normal buffer visualization", EUserInterfaceActionType::RadioButton, FInputChord(EModifierKey::Alt, EKeys::Seven));
+	UI_COMMAND(GoToHLODDistance, "Go to HLOD Distance", "Move camera to the HLOD minimum visible distance", EUserInterfaceActionType::Button, FInputChord(EKeys::SpaceBar));
+}
+
+#undef LOCTEXT_NAMESPACE
diff --git a/Engine/Source/Editor/WorldPartitionEditor/Private/WorldPartition/HLOD/HLODCompareCommands.h b/Engine/Source/Editor/WorldPartitionEditor/Private/WorldPartition/HLOD/HLODCompareCommands.h
new file mode 100644
--- /dev/null
+++ b/Engine/Source/Editor/WorldPartitionEditor/Private/WorldPartition/HLOD/HLODCompareCommands.h
@@ -0,0 +1,27 @@
+// Copyright Epic Games, Inc. All Rights Reserved.
+
+#pragma once
+
+#include "Framework/Commands/Commands.h"
+#include "Framework/Commands/UICommandInfo.h"
+#include "Styling/AppStyle.h"
+
+class FHLODCompareCommands : public TCommands<FHLODCompareCommands>
+{
+public:
+	FHLODCompareCommands()
+		: TCommands<FHLODCompareCommands>(TEXT("HLODCompare"), NSLOCTEXT("Contexts", "HLODCompare", "HLOD Compare"), NAME_None, FAppStyle::GetAppStyleSetName())
+	{
+	}
+
+	virtual void RegisterCommands() override;
+
+	TSharedPtr<FUICommandInfo> ViewModeLit;
+	TSharedPtr<FUICommandInfo> ViewModeWireframe;
+	TSharedPtr<FUICommandInfo> ViewModeBaseColor;
+	TSharedPtr<FUICommandInfo> ViewModeMetallic;
+	TSharedPtr<FUICommandInfo> ViewModeRoughness;
+	TSharedPtr<FUICommandInfo> ViewModeSpecular;
+	TSharedPtr<FUICommandInfo> ViewModeWorldNormal;
+	TSharedPtr<FUICommandInfo> GoToHLODDistance;
+};
diff --git a/Engine/Source/Editor/WorldPartitionEditor/Private/WorldPartition/HLOD/HLODCompareViewportClient.cpp b/Engine/Source/Editor/WorldPartitionEditor/Private/WorldPartition/HLOD/HLODCompareViewportClient.cpp
new file mode 100644
--- /dev/null
+++ b/Engine/Source/Editor/WorldPartitionEditor/Private/WorldPartition/HLOD/HLODCompareViewportClient.cpp
@@ -0,0 +1,76 @@
+// Copyright Epic Games, Inc. All Rights Reserved.
+
+#include "WorldPartition/HLOD/HLODCompareViewportClient.h"
+#include "WorldPartition/HLOD/SHLODCompareWindow.h"
+#include "PreviewScene.h"
+
+FHLODCompareViewportClient::FHLODCompareViewportClient(FPreviewScene* InPreviewScene, const TSharedRef<SEditorViewport>& InEditorViewportWidget, TWeakPtr<SHLODCompareWindow> InCompareWindow)
+	: FEditorViewportClient(nullptr, InPreviewScene, InEditorViewportWidget)
+	, CompareWindow(InCompareWindow)
+{
+	SetRealtime(true);
+	SetViewportType(LVT_Perspective);
+}
+
+void FHLODCompareViewportClient::PerspectiveCameraMoved()
+{
+	FEditorViewportClient::PerspectiveCameraMoved();
+
+	if (!bIsSyncing)
+	{
+		if (TSharedPtr<SHLODCompareWindow> Window = CompareWindow.Pin())
+		{
+			Window->OnCameraMoved(this);
+		}
+	}
+}
+
+void FHLODCompareViewportClient::SetViewMode(EViewModeIndex InViewModeIndex)
+{
+	FEditorViewportClient::SetViewMode(InViewModeIndex);
+
+	if (!bIsSyncing)
+	{
+		if (TSharedPtr<SHLODCompareWindow> Window = CompareWindow.Pin())
+		{
+			Window->OnViewModeChanged(this);
+		}
+	}
+}
+
+void FHLODCompareViewportClient::SyncCameraFrom(const FHLODCompareViewportClient& Other)
+{
+	bIsSyncing = true;
+	SetViewLocation(Other.GetViewLocation());
+	SetViewRotation(Other.GetViewRotation());
+	SetLookAtLocation(Other.GetLookAtLocation());
+	SetOrthoZoom(Other.GetOrthoZoom());
+	bIsSyncing = false;
+}
+
+void FHLODCompareViewportClient::Draw(FViewport* InViewport, FCanvas* Canvas)
+{
+	// Pull buffer visualization mode from sibling right before drawing so it's always in sync
+	if (TSharedPtr<FHLODCompareViewportClient> Sibling = SiblingClient.Pin())
+	{
+		CurrentBufferVisualizationMode = Sibling->CurrentBufferVisualizationMode;
+	}
+
+	// Disable temporal AA in buffer visualization modes to avoid flickering
+	EngineShowFlags.SetTemporalAA(GetViewMode() != VMI_VisualizeBuffer);
+
+	// Tick the preview scene's sky light and reflection captures
+	if (FPreviewScene* Scene = GetPreviewScene())
+	{
+		Scene->UpdateCaptureContents();
+	}
+
+	FEditorViewportClient::Draw(InViewport, Canvas);
+}
+
+void FHLODCompareViewportClient::SyncViewModeFrom(const FHLODCompareViewportClient& Other)
+{
+	bIsSyncing = true;
+	SetViewMode(Other.GetViewMode());
+	bIsSyncing = false;
+}
diff --git a/Engine/Source/Editor/WorldPartitionEditor/Private/WorldPartition/HLOD/HLODCompareViewportClient.h b/Engine/Source/Editor/WorldPartitionEditor/Private/WorldPartition/HLOD/HLODCompareViewportClient.h
new file mode 100644
--- /dev/null
+++ b/Engine/Source/Editor/WorldPartitionEditor/Private/WorldPartition/HLOD/HLODCompareViewportClient.h
@@ -0,0 +1,31 @@
+// Copyright Epic Games, Inc. All Rights Reserved.
+
+#pragma once
+
+#include "EditorViewportClient.h"
+
+class FPreviewScene;
+class SHLODCompareWindow;
+class SEditorViewport;
+
+class FHLODCompareViewportClient : public FEditorViewportClient
+{
+public:
+	FHLODCompareViewportClient(FPreviewScene* InPreviewScene, const TSharedRef<SEditorViewport>& InEditorViewportWidget, TWeakPtr<SHLODCompareWindow> InCompareWindow);
+
+	// FEditorViewportClient interface
+	virtual void PerspectiveCameraMoved() override;
+	virtual void SetViewMode(EViewModeIndex InViewModeIndex) override;
+	virtual void Draw(FViewport* InViewport, FCanvas* Canvas) override;
+
+	void SyncCameraFrom(const FHLODCompareViewportClient& Other);
+	void SyncViewModeFrom(const FHLODCompareViewportClient& Other);
+
+	/** Set a sibling viewport client to sync buffer visualization mode from before each draw. */
+	void SetSiblingClient(TWeakPtr<FHLODCompareViewportClient> InSibling) { SiblingClient = InSibling; }
+
+private:
+	TWeakPtr<SHLODCompareWindow> CompareWindow;
+	TWeakPtr<FHLODCompareViewportClient> SiblingClient;
+	bool bIsSyncing = false;
+};
diff --git a/Engine/Source/Editor/WorldPartitionEditor/Private/WorldPartition/HLOD/SHLODClipWrapper.h b/Engine/Source/Editor/WorldPartitionEditor/Private/WorldPartition/HLOD/SHLODClipWrapper.h
new file mode 100644
--- /dev/null
+++ b/Engine/Source/Editor/WorldPartitionEditor/Private/WorldPartition/HLOD/SHLODClipWrapper.h
@@ -0,0 +1,54 @@
+// Copyright Epic Games, Inc. All Rights Reserved.
+
+#pragma once
+
+#include "Widgets/SCompoundWidget.h"
+#include "Layout/Clipping.h"
+#include "Rendering/DrawElements.h"
+
+/**
+ * A simple compound widget that clips its single child to show only the portion
+ * right of a dynamic wipe position. Used for the HLOD compare wipe overlay.
+ */
+class SHLODClipWrapper : public SCompoundWidget
+{
+public:
+	SLATE_BEGIN_ARGS(SHLODClipWrapper)
+		: _WipePosition(0.5f)
+	{}
+		SLATE_DEFAULT_SLOT(FArguments, Content)
+		SLATE_ATTRIBUTE(float, WipePosition)
+	SLATE_END_ARGS()
+
+	void Construct(const FArguments& InArgs)
+	{
+		WipePosition = InArgs._WipePosition;
+
+		ChildSlot
+		[
+			InArgs._Content.Widget
+		];
+
+		SetVisibility(EVisibility::HitTestInvisible);
+	}
+
+	virtual int32 OnPaint(const FPaintArgs& Args, const FGeometry& AllottedGeometry, const FSlateRect& MyCullingRect, FSlateWindowElementList& OutDrawElements, int32 LayerId, const FWidgetStyle& InWidgetStyle, bool bParentEnabled) const override
+	{
+		const FVector2D AbsPos = AllottedGeometry.GetAbsolutePosition();
+		const FVector2D AbsSize = AllottedGeometry.GetAbsoluteSize();
+		const float WipeAbsX = AbsPos.X + AbsSize.X * WipePosition.Get();
+
+		FSlateRect ClipRect(WipeAbsX, AbsPos.Y, AbsPos.X + AbsSize.X, AbsPos.Y + AbsSize.Y);
+		FSlateClippingZone ClipZone(ClipRect);
+		ClipZone.SetShouldIntersectParent(true);
+
+		OutDrawElements.PushClip(ClipZone);
+		int32 Result = SCompoundWidget::OnPaint(Args, AllottedGeometry, MyCullingRect, OutDrawElements, LayerId, InWidgetStyle, bParentEnabled);
+		OutDrawElements.PopClip();
+
+		return Result;
+	}
+
+private:
+	TAttribute<float> WipePosition;
+};
diff --git a/Engine/Source/Editor/WorldPartitionEditor/Private/WorldPartition/HLOD/SHLODCompareViewport.cpp b/Engine/Source/Editor/WorldPartitionEditor/Private/WorldPartition/HLOD/SHLODCompareViewport.cpp
new file mode 100644
--- /dev/null
+++ b/Engine/Source/Editor/WorldPartitionEditor/Private/WorldPartition/HLOD/SHLODCompareViewport.cpp
@@ -0,0 +1,143 @@
+// Copyright Epic Games, Inc. All Rights Reserved.
+
+#include "WorldPartition/HLOD/SHLODCompareViewport.h"
+#include "WorldPartition/HLOD/HLODCompareViewportClient.h"
+#include "WorldPartition/HLOD/SHLODCompareWindow.h"
+#include "PreviewScene.h"
+#include "ToolMenus.h"
+#include "ViewportToolbar/UnrealEdViewportToolbar.h"
+#include "WorldPartition/HLOD/HLODCompareCommands.h"
+
+#define LOCTEXT_NAMESPACE "SHLODCompareViewport"
+
+static void PopulateHLODCompareViewModesMenu(UToolMenu* InMenu)
+{
+	const FHLODCompareCommands& Commands = FHLODCompareCommands::Get();
+	const FName StyleSetName = FAppStyle::GetAppStyleSetName();
+
+	auto ViewModeIcon = [&StyleSetName](const TCHAR* IconName) -> FSlateIcon
+	{
+		return FSlateIcon(StyleSetName, IconName);
+	};
+
+	// View modes section: Lit, Wireframe
+	{
+		FToolMenuSection& Section = InMenu->AddSection("ViewMode", LOCTEXT("ViewModeHeader", "View Mode"));
+		Section.AddMenuEntry(Commands.ViewModeLit, TAttribute<FText>(), TAttribute<FText>(), ViewModeIcon(TEXT("EditorViewport.LitMode")));
+		Section.AddMenuEntry(Commands.ViewModeWireframe, TAttribute<FText>(), TAttribute<FText>(), ViewModeIcon(TEXT("EditorViewport.WireframeMode")));
+	}
+
+	// Buffer visualization section
+	{
+		FSlateIcon BufferIcon = ViewModeIcon(TEXT("EditorViewport.VisualizeBufferMode"));
+		FToolMenuSection& Section = InMenu->AddSection("BufferVisualization", LOCTEXT("BufferVisualizationHeader", "Buffer Visualization"));
+		Section.AddMenuEntry(Commands.ViewModeBaseColor, TAttribute<FText>(), TAttribute<FText>(), BufferIcon);
+		Section.AddMenuEntry(Commands.ViewModeMetallic, TAttribute<FText>(), TAttribute<FText>(), BufferIcon);
+		Section.AddMenuEntry(Commands.ViewModeRoughness, TAttribute<FText>(), TAttribute<FText>(), BufferIcon);
+		Section.AddMenuEntry(Commands.ViewModeSpecular, TAttribute<FText>(), TAttribute<FText>(), BufferIcon);
+		Section.AddMenuEntry(Commands.ViewModeWorldNormal, TAttribute<FText>(), TAttribute<FText>(), BufferIcon);
+	}
+}
+
+void SHLODCompareViewport::Construct(const FArguments& InArgs)
+{
+	PreviewScene = InArgs._PreviewScene;
+	CompareWindow = InArgs._CompareWindow;
+
+	SEditorViewport::Construct(SEditorViewport::FArguments());
+}
+
+FReply SHLODCompareViewport::OnKeyDown(const FGeometry& MyGeometry, const FKeyEvent& InKeyEvent)
+{
+	// Check the priority command list before the viewport's own commands.
+	// This ensures our HLOD compare shortcuts (Alt+1..7, Space) take precedence
+	// over the standard viewport view mode shortcuts.
+	if (PriorityCommandList.IsValid() && PriorityCommandList->ProcessCommandBindings(InKeyEvent))
+	{
+		return FReply::Handled();
+	}
+
+	// Swallow other Alt+number combos so the viewport's standard bindings don't fire
+	if (InKeyEvent.IsAltDown())
+	{
+		const FKey Key = InKeyEvent.GetKey();
+		if (Key == EKeys::Eight || Key == EKeys::Nine || Key == EKeys::Zero)
+		{
+			return FReply::Handled();
+		}
+	}
+
+	return SEditorViewport::OnKeyDown(MyGeometry, InKeyEvent);
+}
+
+TSharedRef<FEditorViewportClient> SHLODCompareViewport::MakeEditorViewportClient()
+{
+	CompareViewportClient = MakeShareable(new FHLODCompareViewportClient(PreviewScene.Get(), SharedThis(this), CompareWindow));
+	return CompareViewportClient.ToSharedRef();
+}
+
+TSharedRef<SWidget> SHLODCompareViewport::BuildExternalViewModeToolbar(const TSharedPtr<FUICommandList>& InCommandList)
+{
+	const FName ViewportToolbarName = "HLODCompare.ViewportToolbar";
+
+	if (!UToolMenus::Get()->IsMenuRegistered(ViewportToolbarName))
+	{
+		UToolMenu* ViewportToolbarMenu = UToolMenus::Get()->RegisterMenu(
+			ViewportToolbarName, NAME_None, EMultiBoxType::SlimHorizontalToolBar
+		);
+		ViewportToolbarMenu->StyleName = "ViewportToolbar";
+
+		FToolMenuSection& ToolbarSection = ViewportToolbarMenu->AddSection("Right");
+		ToolbarSection.Alignment = EToolMenuSectionAlign::Last;
+
+		// Add a dynamic submenu entry with label/icon that reflects the current view mode
+		ToolbarSection.AddEntry(FToolMenuEntry::InitDynamicEntry("DynamicViewModes",
+			FNewToolMenuSectionDelegate::CreateLambda(
+				[](FToolMenuSection& InSection)
+				{
+					UUnrealEdViewportToolbarContext* Context = InSection.FindContext<UUnrealEdViewportToolbarContext>();
+					if (!Context)
+					{
+						return;
+					}
+
+					TAttribute<FText> LabelAttribute = TAttribute<FText>::CreateLambda(
+						[WeakViewport = Context->Viewport]()
+						{
+							return UE::UnrealEd::GetViewModesSubmenuLabel(WeakViewport);
+						});
+
+					TAttribute<FSlateIcon> IconAttribute = TAttribute<FSlateIcon>::CreateLambda(
+						[WeakViewport = Context->Viewport]()
+						{
+							return UE::UnrealEd::GetViewModesSubmenuIcon(WeakViewport);
+						});
+
+					InSection.AddSubMenu(
+						"ViewModes",
+						LabelAttribute,
+						LOCTEXT("ViewModesSubmenuTooltip", "View mode settings for the current viewport."),
+						FNewToolMenuDelegate::CreateStatic(&PopulateHLODCompareViewModesMenu),
+						false,
+						IconAttribute
+					);
+				})));
+	}
+
+	FToolMenuContext ViewportToolbarContext;
+	{
+		if (InCommandList.IsValid())
+		{
+			ViewportToolbarContext.AppendCommandList(InCommandList);
+		}
+
+		UUnrealEdViewportToolbarContext* ContextObject =
+			UE::UnrealEd::CreateViewportToolbarDefaultContext(SharedThis(this));
+		ContextObject->bShowCoordinateSystemControls = false;
+		ViewportToolbarContext.AddObject(ContextObject);
+	}
+
+	return UToolMenus::Get()->GenerateWidget(ViewportToolbarName, ViewportToolbarContext);
+}
+
+#undef LOCTEXT_NAMESPACE
diff --git a/Engine/Source/Editor/WorldPartitionEditor/Private/WorldPartition/HLOD/SHLODCompareViewport.h b/Engine/Source/Editor/WorldPartitionEditor/Private/WorldPartition/HLOD/SHLODCompareViewport.h
new file mode 100644
--- /dev/null
+++ b/Engine/Source/Editor/WorldPartitionEditor/Private/WorldPartition/HLOD/SHLODCompareViewport.h
@@ -0,0 +1,39 @@
+// Copyright Epic Games, Inc. All Rights Reserved.
+
+#pragma once
+
+#include "SEditorViewport.h"
+
+class FPreviewScene;
+class SHLODCompareWindow;
+class FHLODCompareViewportClient;
+
+class SHLODCompareViewport : public SEditorViewport
+{
+public:
+	SLATE_BEGIN_ARGS(SHLODCompareViewport) {}
+		SLATE_ARGUMENT(TSharedPtr<FPreviewScene>, PreviewScene)
+		SLATE_ARGUMENT(TWeakPtr<SHLODCompareWindow>, CompareWindow)
+	SLATE_END_ARGS()
+
+	void Construct(const FArguments& InArgs);
+
+	TSharedPtr<FHLODCompareViewportClient> GetCompareViewportClient() const { return CompareViewportClient; }
+
+	/** Build a standalone view mode toolbar widget for this viewport (used by SHLODCompareWindow). */
+	TSharedRef<SWidget> BuildExternalViewModeToolbar(const TSharedPtr<FUICommandList>& InCommandList);
+
+	/** Set the priority command list that is checked before the viewport's own commands. */
+	void SetPriorityCommandList(const TSharedPtr<FUICommandList>& InCommandList) { PriorityCommandList = InCommandList; }
+
+protected:
+	virtual TSharedRef<FEditorViewportClient> MakeEditorViewportClient() override;
+	virtual TSharedPtr<SWidget> BuildViewportToolbar() override { return nullptr; }
+	virtual FReply OnKeyDown(const FGeometry& MyGeometry, const FKeyEvent& InKeyEvent) override;
+
+private:
+	TSharedPtr<FPreviewScene> PreviewScene;
+	TWeakPtr<SHLODCompareWindow> CompareWindow;
+	TSharedPtr<FHLODCompareViewportClient> CompareViewportClient;
+	TSharedPtr<FUICommandList> PriorityCommandList;
+};
diff --git a/Engine/Source/Editor/WorldPartitionEditor/Private/WorldPartition/HLOD/SHLODCompareWindow.cpp b/Engine/Source/Editor/WorldPartitionEditor/Private/WorldPartition/HLOD/SHLODCompareWindow.cpp
new file mode 100644
--- /dev/null
+++ b/Engine/Source/Editor/WorldPartitionEditor/Private/WorldPartition/HLOD/SHLODCompareWindow.cpp
@@ -0,0 +1,831 @@
+// Copyright Epic Games, Inc. All Rights Reserved.
+
+#include "WorldPartition/HLOD/SHLODCompareWindow.h"
+#include "WorldPartition/HLOD/SHLODCompareViewport.h"
+#include "WorldPartition/HLOD/HLODCompareViewportClient.h"
+#include "WorldPartition/HLOD/SHLODClipWrapper.h"
+#include "WorldPartition/HLOD/SHLODWipeHandle.h"
+#include "PreviewScene.h"
+#include "Components/PrimitiveComponent.h"
+#include "Components/DirectionalLightComponent.h"
+#include "Components/SkyLightComponent.h"
+#include "Editor.h"
+#include "Framework/Application/SlateApplication.h"
+#include "Misc/ScopedSlowTask.h"
+#include "LevelEditorViewport.h"
+#include "Engine/World.h"
+#include "GameFramework/WorldSettings.h"
+#include "EngineUtils.h"
+#include "Widgets/SOverlay.h"
+#include "Widgets/SCanvas.h"
+#include "Widgets/Text/STextBlock.h"
+#include "Widgets/Input/SButton.h"
+#include "Widgets/Layout/SBox.h"
+#include "WorldPartition/HLOD/HLODActor.h"
+#include "WorldPartition/HLOD/HLODCompareCommands.h"
+#include "WorldPartition/HLOD/HLODSourceActors.h"
+#include "WorldPartition/HLOD/SHLODDistanceRuler.h"
+
+#define LOCTEXT_NAMESPACE "SHLODCompareWindow"
+
+void SHLODCompareWindow::Construct(const FArguments& InArgs)
+{
+	HLODActors = InArgs._HLODActors;
+
+	SourcePreviewScene = MakeShared<FPreviewScene>(FPreviewScene::ConstructionValues()
+		.SetCreateDefaultLighting(true)
+		.SetCreatePhysicsScene(false)
+		.ShouldSimulatePhysics(false)
+		.SetTransactional(false));
+
+	HLODPreviewScene = MakeShared<FPreviewScene>(FPreviewScene::ConstructionValues()
+		.SetCreateDefaultLighting(true)
+		.SetCreatePhysicsScene(false)
+		.ShouldSimulatePhysics(false)
+		.SetTransactional(false));
+
+	TSharedPtr<SHLODCompareWindow> ThisPtr = SharedThis(this);
+
+	ChildSlot
+	[
+		SNew(SVerticalBox)
+
+		// Labels row
+		+ SVerticalBox::Slot()
+		.AutoHeight()
+		.Padding(4.0f, 2.0f)
+		[
+			SNew(SHorizontalBox)
+			+ SHorizontalBox::Slot()
+			.FillWidth(1.0f)
+			.HAlign(HAlign_Left)
+			.Padding(8.0f, 0.0f)
+			[
+				SNew(STextBlock)
+				.Text(LOCTEXT("SourceActorsLabel", "Source Actors"))
+				.Font(FCoreStyle::GetDefaultFontStyle("Bold", 11))
+			]
+			+ SHorizontalBox::Slot()
+			.FillWidth(1.0f)
+			.HAlign(HAlign_Right)
+			.Padding(8.0f, 0.0f)
+			[
+				SNew(STextBlock)
+				.Text(LOCTEXT("HLODLabel", "HLOD"))
+				.Font(FCoreStyle::GetDefaultFontStyle("Bold", 11))
+			]
+		]
+
+		// Toolbar row: view mode selector (left) + Go to HLOD distance button (right)
+		+ SVerticalBox::Slot()
+		.AutoHeight()
+		.Padding(4.0f, 0.0f)
+		[
+			SNew(SHorizontalBox)
+			+ SHorizontalBox::Slot()
+			.AutoWidth()
+			[
+				SAssignNew(ViewModeToolbarContainer, SBox)
+			]
+			+ SHorizontalBox::Slot()
+			.FillWidth(1.0f)
+			.HAlign(HAlign_Right)
+			.VAlign(VAlign_Center)
+			.Padding(8.0f, 0.0f)
+			[
+				SNew(SButton)
+				.Text(LOCTEXT("GoToHLODDistance", "Go to HLOD Distance"))
+				.ToolTipText(LOCTEXT("GoToHLODDistance_Tooltip", "Move camera to the distance at which the HLOD first becomes visible"))
+				.OnClicked(this, &SHLODCompareWindow::OnGoToMinVisibleDistance)
+			]
+		]
+
+		// Wipe overlay viewports
+		+ SVerticalBox::Slot()
+		.FillHeight(1.0f)
+		[
+			SNew(SOverlay)
+
+			// Z=0: Source viewport (full area, receives all input)
+			+ SOverlay::Slot()
+			[
+				SAssignNew(SourceViewport, SHLODCompareViewport)
+				.PreviewScene(SourcePreviewScene)
+				.CompareWindow(ThisPtr)
+			]
+
+			// Z=1: Clipped HLOD viewport (HitTestInvisible so input passes through)
+			+ SOverlay::Slot()
+			.HAlign(HAlign_Fill)
+			.VAlign(VAlign_Fill)
+			[
+				SAssignNew(ClipWrapper, SHLODClipWrapper)
+				.WipePosition(TAttribute<float>::CreateLambda([this]() { return WipePosition; }))
+				[
+					SAssignNew(HLODViewport, SHLODCompareViewport)
+					.PreviewScene(HLODPreviewScene)
+					.CompareWindow(ThisPtr)
+				]
+			]
+
+			// Z=2: Wipe handle via SCanvas
+			+ SOverlay::Slot()
+			.HAlign(HAlign_Fill)
+			.VAlign(VAlign_Fill)
+			[
+				SNew(SCanvas)
+				+ SCanvas::Slot()
+				.Position(TAttribute<FVector2D>::CreateLambda([this]()
+				{
+					// Get canvas geometry to compute pixel position
+					float Width = 1.0f;
+					if (SourceViewport.IsValid())
+					{
+						Width = FMath::Max(SourceViewport->GetTickSpaceGeometry().GetLocalSize().X, 1.0f);
+					}
+					float HandleX = Width * WipePosition - 6.0f; // center 12px handle on wipe line
+					return FVector2D(HandleX, 0.0);
+				}))
+				.Size(TAttribute<FVector2D>::CreateLambda([this]()
+				{
+					float Height = 1.0f;
+					if (SourceViewport.IsValid())
+					{
+						Height = FMath::Max(SourceViewport->GetTickSpaceGeometry().GetLocalSize().Y, 1.0f);
+					}
+					return FVector2D(12.0, Height);
+				}))
+				[
+					SNew(SHLODWipeHandle)
+					.OnWipePositionChanged(FOnWipePositionChanged::CreateSP(this, &SHLODCompareWindow::OnWipePositionChanged))
+				]
+			]
+		]
+
+		// Status text (hidden by default)
+		+ SVerticalBox::Slot()
+		.AutoHeight()
+		.Padding(4.0f, 2.0f)
+		[
+			SAssignNew(StatusText, STextBlock)
+			.ColorAndOpacity(FLinearColor::Yellow)
+			.Visibility(EVisibility::Collapsed)
+		]
+
+		// Distance ruler (bottom of window)
+		+ SVerticalBox::Slot()
+		.AutoHeight()
+		.Padding(4.0f, 2.0f)
+		[
+			SNew(SHLODDistanceRuler)
+			.SourceMinDrawDistance(TAttribute<float>::CreateSP(this, &SHLODCompareWindow::GetSourceMinDrawDistance))
+			.HLODMinVisibleDistance(TAttribute<float>::CreateSP(this, &SHLODCompareWindow::GetHLODMinVisibleDistance))
+			.CameraDistance(TAttribute<float>::CreateSP(this, &SHLODCompareWindow::GetCameraDistanceToBounds))
+			.SourceZoneLabel(TAttribute<FString>::CreateSP(this, &SHLODCompareWindow::GetSourceZoneLabel))
+			.HLODZoneLabel(TAttribute<FString>::CreateSP(this, &SHLODCompareWindow::GetHLODZoneLabel))
+			.WorldToMeters(TAttribute<float>::CreateLambda([]()
+			{
+				UWorld* World = GEditor ? GEditor->GetEditorWorldContext().World() : nullptr;
+				return (World && World->GetWorldSettings()) ? World->GetWorldSettings()->WorldToMeters : 100.0f;
+			}))
+		]
+	];
+
+	// Register and bind keyboard shortcuts (must happen before building toolbar so command list is available)
+	BindCommands();
+
+	// Now that the source viewport is constructed, build and insert the view mode toolbar
+	if (SourceViewport.IsValid())
+	{
+		ViewModeToolbarContainer->SetContent(SourceViewport->BuildExternalViewModeToolbar(CommandList));
+	}
+
+	// Set our command list as priority on the viewport so our shortcuts take precedence over standard ones
+	if (SourceViewport.IsValid() && CommandList.IsValid())
+	{
+		SourceViewport->SetPriorityCommandList(CommandList);
+	}
+
+	// Wire up sibling relationship so HLOD client pulls buffer viz name from source before drawing
+	if (SourceViewport.IsValid() && HLODViewport.IsValid())
+	{
+		TSharedPtr<FHLODCompareViewportClient> SourceClient = SourceViewport->GetCompareViewportClient();
+		TSharedPtr<FHLODCompareViewportClient> HLODClient = HLODViewport->GetCompareViewportClient();
+		if (SourceClient.IsValid() && HLODClient.IsValid())
+		{
+			HLODClient->SetSiblingClient(SourceClient);
+		}
+	}
+
+	// Populate scenes
+	bool bSourceOk = PopulateSourceScene();
+	bool bHLODOk = PopulateHLODScene();
+
+	if (!bSourceOk)
+	{
+		StatusText->SetText(LOCTEXT("SourceLoadFailed", "Warning: Failed to load source actors. Left side of wipe may be empty."));
+		StatusText->SetVisibility(EVisibility::Visible);
+	}
+	else if (!bHLODOk)
+	{
+		StatusText->SetText(LOCTEXT("HLODLoadFailed", "Warning: HLOD actor has no visual components. Right side of wipe may be empty."));
+		StatusText->SetVisibility(EVisibility::Visible);
+	}
+
+	CopyWorldLighting();
+	FocusCamerasOnContent();
+
+	// Defer keyboard focus to the source viewport until the window is visible
+	if (SourceViewport.IsValid())
+	{
+		TWeakPtr<SHLODCompareViewport> WeakViewport = SourceViewport;
+		RegisterActiveTimer(0.0f, FWidgetActiveTimerDelegate::CreateLambda(
+			[WeakViewport](double, float) -> EActiveTimerReturnType
+			{
+				if (TSharedPtr<SHLODCompareViewport> Viewport = WeakViewport.Pin())
+				{
+					FSlateApplication::Get().SetKeyboardFocus(Viewport);
+				}
+				return EActiveTimerReturnType::Stop;
+			}));
+	}
+}
+
+SHLODCompareWindow::~SHLODCompareWindow()
+{
+	// Remove components from preview scenes before destruction
+	for (UActorComponent* Component : SourceSceneComponents)
+	{
+		if (Component)
+		{
+			SourcePreviewScene->RemoveComponent(Component);
+		}
+	}
+	SourceSceneComponents.Empty();
+
+	for (UActorComponent* Component : HLODSceneComponents)
+	{
+		if (Component)
+		{
+			HLODPreviewScene->RemoveComponent(Component);
+		}
+	}
+	HLODSceneComponents.Empty();
+}
+
+bool SHLODCompareWindow::PopulateSourceScene()
+{
+	UWorld* PreviewWorld = SourcePreviewScene->GetWorld();
+
+	FScopedSlowTask SlowTask(static_cast<float>(HLODActors.Num()), LOCTEXT("LoadingSourceActors", "Loading HLOD source actors..."));
+	SlowTask.MakeDialog();
+
+	for (const TWeakObjectPtr<AWorldPartitionHLOD>& WeakHLODActor : HLODActors)
+	{
+		AWorldPartitionHLOD* HLODActorPtr = WeakHLODActor.Get();
+		if (!HLODActorPtr)
+		{
+			continue;
+		}
+
+		SlowTask.EnterProgressFrame(1.0f, FText::Format(LOCTEXT("LoadingSourceActorsStep", "Loading source actors for {0}"), FText::FromString(HLODActorPtr->GetActorLabel())));
+
+		UWorldPartitionHLODSourceActors* SourceActors = HLODActorPtr->GetSourceActors();
+		if (!SourceActors)
+		{
+			continue;
+		}
+
+		// Load source actors directly into the preview scene's world
+		bool bDirty = false;
+		if (!SourceActors->LoadSourceActors(bDirty, PreviewWorld))
+		{
+			continue;
+		}
+
+		// Iterate loaded actors and add their primitive components to the preview scene
+		for (TActorIterator<AActor> It(PreviewWorld); It; ++It)
+		{
+			AActor* Actor = *It;
+			if (!SourceActors->IsHLODRelevant(Actor))
+			{
+				continue;
+			}
+
+			TArray<UPrimitiveComponent*> PrimitiveComponents;
+			Actor->GetComponents<UPrimitiveComponent>(PrimitiveComponents);
+			for (UPrimitiveComponent* PrimComp : PrimitiveComponents)
+			{
+				if (PrimComp && !SourceSceneComponents.Contains(PrimComp))
+				{
+					SourcePreviewScene->AddComponent(PrimComp, PrimComp->GetComponentTransform());
+					SourceSceneComponents.Add(PrimComp);
+				}
+			}
+		}
+	}
+
+	return SourceSceneComponents.Num() > 0;
+}
+
+bool SHLODCompareWindow::PopulateHLODScene()
+{
+	for (const TWeakObjectPtr<AWorldPartitionHLOD>& WeakHLODActor : HLODActors)
+	{
+		AWorldPartitionHLOD* HLODActorPtr = WeakHLODActor.Get();
+		if (!HLODActorPtr)
+		{
+			continue;
+		}
+
+		const TArray<UActorComponent*>& InstanceComponents = HLODActorPtr->GetInstanceComponents();
+		for (UActorComponent* Component : InstanceComponents)
+		{
+			USceneComponent* SceneComp = Cast<USceneComponent>(Component);
+			if (!SceneComp)
+			{
+				continue;
+			}
+
+			// Duplicate the component into the HLOD preview scene's world so that the HLOD actor
+			// is not affected and the preview scene takes ownership.
+			USceneComponent* DuplicatedComp = DuplicateObject<USceneComponent>(SceneComp, HLODPreviewScene->GetWorld());
+			if (DuplicatedComp)
+			{
+				HLODPreviewScene->AddComponent(DuplicatedComp, SceneComp->GetComponentTransform());
+				HLODSceneComponents.Add(DuplicatedComp);
+			}
+		}
+	}
+
+	return HLODSceneComponents.Num() > 0;
+}
+
+void SHLODCompareWindow::OnCameraMoved(FHLODCompareViewportClient* Source)
+{
+	if (!SourceViewport.IsValid() || !HLODViewport.IsValid())
+	{
+		return;
+	}
+
+	TSharedPtr<FHLODCompareViewportClient> SourceClient = SourceViewport->GetCompareViewportClient();
+	TSharedPtr<FHLODCompareViewportClient> HLODClient = HLODViewport->GetCompareViewportClient();
+
+	if (!SourceClient.IsValid() || !HLODClient.IsValid())
+	{
+		return;
+	}
+
+	// Sync the sibling viewport to match the one the user is controlling
+	if (Source == SourceClient.Get())
+	{
+		HLODClient->SyncCameraFrom(*SourceClient);
+	}
+	else if (Source == HLODClient.Get())
+	{
+		SourceClient->SyncCameraFrom(*HLODClient);
+	}
+}
+
+void SHLODCompareWindow::OnViewModeChanged(FHLODCompareViewportClient* Source)
+{
+	if (!SourceViewport.IsValid() || !HLODViewport.IsValid())
+	{
+		return;
+	}
+
+	TSharedPtr<FHLODCompareViewportClient> SourceClient = SourceViewport->GetCompareViewportClient();
+	TSharedPtr<FHLODCompareViewportClient> HLODClient = HLODViewport->GetCompareViewportClient();
+
+	if (!SourceClient.IsValid() || !HLODClient.IsValid())
+	{
+		return;
+	}
+
+	// Sync view mode to the sibling viewport
+	if (Source == SourceClient.Get())
+	{
+		HLODClient->SyncViewModeFrom(*SourceClient);
+	}
+	else if (Source == HLODClient.Get())
+	{
+		SourceClient->SyncViewModeFrom(*HLODClient);
+	}
+
+}
+
+void SHLODCompareWindow::CopyWorldLighting()
+{
+	UWorld* EditorWorld = GEditor ? GEditor->GetEditorWorldContext().World() : nullptr;
+	if (!EditorWorld)
+	{
+		return;
+	}
+
+	auto ApplyToScene = [](FPreviewScene* Scene, UDirectionalLightComponent* WorldDirLight, USkyLightComponent* WorldSkyLight)
+	{
+		if (!Scene)
+		{
+			return;
+		}
+
+		if (WorldDirLight && Scene->DirectionalLight)
+		{
+			Scene->SetLightDirection(WorldDirLight->GetComponentRotation());
+			Scene->SetLightBrightness(WorldDirLight->Intensity);
+			Scene->SetLightColor(WorldDirLight->LightColor);
+		}
+
+		if (WorldSkyLight && Scene->SkyLight)
+		{
+			Scene->SetSkyBrightness(WorldSkyLight->Intensity);
+			if (WorldSkyLight->SourceType == ESkyLightSourceType::SLS_SpecifiedCubemap && WorldSkyLight->Cubemap)
+			{
+				Scene->SetSkyCubemap(WorldSkyLight->Cubemap);
+			}
+			Scene->SkyLight->bLowerHemisphereIsBlack = WorldSkyLight->bLowerHemisphereIsBlack;
+		}
+	};
+
+	// Find the first directional light and sky light in the editor world
+	UDirectionalLightComponent* WorldDirLight = nullptr;
+	USkyLightComponent* WorldSkyLight = nullptr;
+
+	for (TActorIterator<AActor> It(EditorWorld); It; ++It)
+	{
+		if (!WorldDirLight)
+		{
+			WorldDirLight = It->FindComponentByClass<UDirectionalLightComponent>();
+		}
+		if (!WorldSkyLight)
+		{
+			WorldSkyLight = It->FindComponentByClass<USkyLightComponent>();
+		}
+		if (WorldDirLight && WorldSkyLight)
+		{
+			break;
+		}
+	}
+
+	ApplyToScene(SourcePreviewScene.Get(), WorldDirLight, WorldSkyLight);
+	ApplyToScene(HLODPreviewScene.Get(), WorldDirLight, WorldSkyLight);
+}
+
+void SHLODCompareWindow::FocusCamerasOnContent()
+{
+	// Compute bounding box from source actors only
+	FBox SourceBounds(ForceInit);
+	for (UActorComponent* Component : SourceSceneComponents)
+	{
+		if (UPrimitiveComponent* PrimComp = Cast<UPrimitiveComponent>(Component))
+		{
+			SourceBounds += PrimComp->Bounds.GetBox();
+		}
+	}
+
+	if (!SourceBounds.IsValid)
+	{
+		return;
+	}
+
+	const FVector Center = SourceBounds.GetCenter();
+
+	// Start from the active editor viewport's camera
+	FVector CameraLocation = Center - FVector(1, 0, 0) * 500.0f;
+	FRotator CameraRotation = FRotator::ZeroRotator;
+
+	if (GCurrentLevelEditingViewportClient)
+	{
+		CameraLocation = GCurrentLevelEditingViewportClient->GetViewLocation();
+		CameraRotation = GCurrentLevelEditingViewportClient->GetViewRotation();
+	}
+
+	// Clamp distance: if camera is further than MinVisibleDistance from the bounds, bring it in.
+	// Use the minimum MinVisibleDistance across all HLOD actors.
+	float MinVisibleDistance = 0.0f;
+	for (const TWeakObjectPtr<AWorldPartitionHLOD>& WeakHLODActor : HLODActors)
+	{
+		if (AWorldPartitionHLOD* HLODActorPtr = WeakHLODActor.Get())
+		{
+			float Dist = static_cast<float>(HLODActorPtr->GetMinVisibleDistance());
+			if (MinVisibleDistance == 0.0f || (Dist > 0.0f && Dist < MinVisibleDistance))
+			{
+				MinVisibleDistance = Dist;
+			}
+		}
+	}
+	if (MinVisibleDistance > 0.0f)
+	{
+		const float DistanceToBounds = FMath::Sqrt(SourceBounds.ComputeSquaredDistanceToPoint(CameraLocation));
+		const float BoundsRadius = SourceBounds.GetExtent().Size();
+		const float DesiredDistance = BoundsRadius + MinVisibleDistance;
+		const FVector ToCamera = CameraLocation - Center;
+		if (ToCamera.Size() > DesiredDistance)
+		{
+			CameraLocation = Center + ToCamera.GetSafeNormal() * DesiredDistance;
+		}
+	}
+
+	TSharedPtr<FHLODCompareViewportClient> SourceClient = SourceViewport.IsValid() ? SourceViewport->GetCompareViewportClient() : nullptr;
+	TSharedPtr<FHLODCompareViewportClient> HLODClient = HLODViewport.IsValid() ? HLODViewport->GetCompareViewportClient() : nullptr;
+
+	auto ApplyCamera = [&](const TSharedPtr<FHLODCompareViewportClient>& ViewportClient)
+	{
+		if (ViewportClient.IsValid())
+		{
+			ViewportClient->SetViewLocation(CameraLocation);
+			ViewportClient->SetViewRotation(CameraRotation);
+			ViewportClient->SetLookAtLocation(Center);
+		}
+	};
+
+	ApplyCamera(SourceClient);
+	ApplyCamera(HLODClient);
+}
+
+float SHLODCompareWindow::GetCameraDistanceToBounds() const
+{
+	TSharedPtr<FHLODCompareViewportClient> Client = SourceViewport.IsValid() ? SourceViewport->GetCompareViewportClient() : nullptr;
+	if (!Client.IsValid())
+	{
+		return 0.0f;
+	}
+
+	FBox SourceBounds(ForceInit);
+	for (UActorComponent* Component : SourceSceneComponents)
+	{
+		if (UPrimitiveComponent* PrimComp = Cast<UPrimitiveComponent>(Component))
+		{
+			SourceBounds += PrimComp->Bounds.GetBox();
+		}
+	}
+
+	if (!SourceBounds.IsValid)
+	{
+		return 0.0f;
+	}
+
+	const float BoundsRadius = SourceBounds.GetExtent().Size();
+	const float DistToCenter = FVector::Dist(Client->GetViewLocation(), SourceBounds.GetCenter());
+	return FMath::Max(DistToCenter - BoundsRadius, 0.0f);
+}
+
+float SHLODCompareWindow::GetSourceMinDrawDistance() const
+{
+	// If the source actors are HLODs (LODLevel > 0), find them in the preview world
+	// and read their MinVisibleDistance directly.
+	if (HLODActors.Num() > 0 && HLODActors[0].IsValid() && HLODActors[0]->GetLODLevel() > 0)
+	{
+		UWorld* PreviewWorld = SourcePreviewScene.IsValid() ? SourcePreviewScene->GetWorld() : nullptr;
+		if (PreviewWorld)
+		{
+			float MaxMinVisible = 0.0f;
+			for (TActorIterator<AWorldPartitionHLOD> It(PreviewWorld); It; ++It)
+			{
+				MaxMinVisible = FMath::Max(MaxMinVisible, static_cast<float>(It->GetMinVisibleDistance()));
+			}
+			if (MaxMinVisible > 0.0f)
+			{
+				return MaxMinVisible;
+			}
+		}
+	}
+
+	// For regular geometry, use the bounds radius as the reference point on the ruler
+	FBox SourceBounds(ForceInit);
+	for (UActorComponent* Component : SourceSceneComponents)
+	{
+		if (UPrimitiveComponent* PrimComp = Cast<UPrimitiveComponent>(Component))
+		{
+			SourceBounds += PrimComp->Bounds.GetBox();
+		}
+	}
+	if (SourceBounds.IsValid)
+	{
+		return SourceBounds.GetExtent().Size();
+	}
+
+	return 0.0f;
+}
+
+float SHLODCompareWindow::GetHLODMinVisibleDistance() const
+{
+	float MinDist = 0.0f;
+	for (const TWeakObjectPtr<AWorldPartitionHLOD>& WeakHLODActor : HLODActors)
+	{
+		if (AWorldPartitionHLOD* HLODActorPtr = WeakHLODActor.Get())
+		{
+			float Dist = static_cast<float>(HLODActorPtr->GetMinVisibleDistance());
+			if (MinDist == 0.0f || (Dist > 0.0f && Dist < MinDist))
+			{
+				MinDist = Dist;
+			}
+		}
+	}
+	return MinDist;
+}
+
+FString SHLODCompareWindow::GetSourceZoneLabel() const
+{
+	// If the HLOD is level N, the source is either "Actors" (N==0) or "HLOD(N-1)"
+	if (HLODActors.Num() > 0 && HLODActors[0].IsValid())
+	{
+		uint32 LODLevel = HLODActors[0]->GetLODLevel();
+		if (LODLevel > 0)
+		{
+			return FString::Printf(TEXT("HLOD%u"), LODLevel - 1);
+		}
+	}
+	return TEXT("Actors");
+}
+
+FString SHLODCompareWindow::GetHLODZoneLabel() const
+{
+	UWorld* World = GEditor ? GEditor->GetEditorWorldContext().World() : nullptr;
+	float W2M = (World && World->GetWorldSettings()) ? World->GetWorldSettings()->WorldToMeters : 100.0f;
+
+	if (HLODActors.Num() > 0 && HLODActors[0].IsValid())
+	{
+		uint32 LODLevel = HLODActors[0]->GetLODLevel();
+		float Dist = GetHLODMinVisibleDistance();
+		return FString::Printf(TEXT("HLOD%u (%.0fm)"), LODLevel, Dist / W2M);
+	}
+	return TEXT("HLOD");
+}
+
+void SHLODCompareWindow::SetCompareViewMode(EViewModeIndex ViewMode, FName BufferVizMode)
+{
+	TSharedPtr<FHLODCompareViewportClient> SourceClient = SourceViewport.IsValid() ? SourceViewport->GetCompareViewportClient() : nullptr;
+	TSharedPtr<FHLODCompareViewportClient> HLODClient = HLODViewport.IsValid() ? HLODViewport->GetCompareViewportClient() : nullptr;
+
+	if (SourceClient.IsValid())
+	{
+		if (ViewMode == VMI_VisualizeBuffer)
+		{
+			SourceClient->ChangeBufferVisualizationMode(BufferVizMode);
+		}
+		else
+		{
+			SourceClient->SetViewMode(ViewMode);
+		}
+	}
+
+	if (HLODClient.IsValid())
+	{
+		if (ViewMode == VMI_VisualizeBuffer)
+		{
+			HLODClient->ChangeBufferVisualizationMode(BufferVizMode);
+		}
+		else
+		{
+			HLODClient->SetViewMode(ViewMode);
+		}
+	}
+}
+
+void SHLODCompareWindow::BindCommands()
+{
+	FHLODCompareCommands::Register();
+	const FHLODCompareCommands& Commands = FHLODCompareCommands::Get();
+
+	CommandList = MakeShared<FUICommandList>();
+
+	auto IsViewMode = [this](EViewModeIndex ViewMode, FName BufferVizMode) -> bool
+	{
+		TSharedPtr<FHLODCompareViewportClient> Client = SourceViewport.IsValid() ? SourceViewport->GetCompareViewportClient() : nullptr;
+		if (!Client.IsValid())
+		{
+			return false;
+		}
+		if (ViewMode == VMI_VisualizeBuffer)
+		{
+			return Client->GetViewMode() == VMI_VisualizeBuffer && Client->CurrentBufferVisualizationMode == BufferVizMode;
+		}
+		return Client->GetViewMode() == ViewMode;
+	};
+
+	auto MapViewMode = [&](const TSharedPtr<FUICommandInfo>& Command, EViewModeIndex ViewMode, FName BufferVizMode = NAME_None)
+	{
+		CommandList->MapAction(Command,
+			FExecuteAction::CreateSP(this, &SHLODCompareWindow::SetCompareViewMode, ViewMode, BufferVizMode),
+			FCanExecuteAction(),
+			FIsActionChecked::CreateLambda([this, ViewMode, BufferVizMode, IsViewMode]() { return IsViewMode(ViewMode, BufferVizMode); }));
+	};
+
+	MapViewMode(Commands.ViewModeLit, VMI_Lit);
+	MapViewMode(Commands.ViewModeWireframe, VMI_BrushWireframe);
+	MapViewMode(Commands.ViewModeBaseColor, VMI_VisualizeBuffer, FName(TEXT("BaseColor")));
+	MapViewMode(Commands.ViewModeMetallic, VMI_VisualizeBuffer, FName(TEXT("Metallic")));
+	MapViewMode(Commands.ViewModeRoughness, VMI_VisualizeBuffer, FName(TEXT("Roughness")));
+	MapViewMode(Commands.ViewModeSpecular, VMI_VisualizeBuffer, FName(TEXT("Specular")));
+	MapViewMode(Commands.ViewModeWorldNormal, VMI_VisualizeBuffer, FName(TEXT("WorldNormal")));
+	CommandList->MapAction(Commands.GoToHLODDistance,
+		FExecuteAction::CreateLambda([this]() { OnGoToMinVisibleDistance(); }));
+}
+
+FReply SHLODCompareWindow::OnKeyDown(const FGeometry& MyGeometry, const FKeyEvent& InKeyEvent)
+{
+	// Let the command list handle registered shortcuts
+	if (CommandList.IsValid() && CommandList->ProcessCommandBindings(InKeyEvent))
+	{
+		return FReply::Handled();
+	}
+
+	// Swallow other Alt+number combos so they don't reach the viewport
+	if (InKeyEvent.IsAltDown())
+	{
+		const FKey Key = InKeyEvent.GetKey();
+		if (Key == EKeys::Eight || Key == EKeys::Nine || Key == EKeys::Zero)
+		{
+			return FReply::Handled();
+		}
+	}
+
+	return SCompoundWidget::OnKeyDown(MyGeometry, InKeyEvent);
+}
+
+void SHLODCompareWindow::OnWipePositionChanged(float NewPosition)
+{
+	WipePosition = FMath::Clamp(NewPosition, 0.0f, 1.0f);
+}
+
+FReply SHLODCompareWindow::OnGoToMinVisibleDistance()
+{
+	// Compute bounds from source actors (used for distance computation)
+	FBox SourceBounds(ForceInit);
+	for (UActorComponent* Component : SourceSceneComponents)
+	{
+		if (UPrimitiveComponent* PrimComp = Cast<UPrimitiveComponent>(Component))
+		{
+			SourceBounds += PrimComp->Bounds.GetBox();
+		}
+	}
+
+	if (!SourceBounds.IsValid)
+	{
+		return FReply::Handled();
+	}
+
+	// Compute HLOD bounds center (used as look-at target)
+	FBox HLODBounds(ForceInit);
+	for (UActorComponent* Component : HLODSceneComponents)
+	{
+		if (UPrimitiveComponent* PrimComp = Cast<UPrimitiveComponent>(Component))
+		{
+			HLODBounds += PrimComp->Bounds.GetBox();
+		}
+	}
+
+	const FVector LookAtCenter = HLODBounds.IsValid ? HLODBounds.GetCenter() : SourceBounds.GetCenter();
+	const FVector SourceCenter = SourceBounds.GetCenter();
+
+	// Find the minimum MinVisibleDistance across all HLOD actors
+	float MinVisibleDistance = 0.0f;
+	for (const TWeakObjectPtr<AWorldPartitionHLOD>& WeakHLODActor : HLODActors)
+	{
+		if (AWorldPartitionHLOD* HLODActorPtr = WeakHLODActor.Get())
+		{
+			float Dist = static_cast<float>(HLODActorPtr->GetMinVisibleDistance());
+			if (MinVisibleDistance == 0.0f || (Dist > 0.0f && Dist < MinVisibleDistance))
+			{
+				MinVisibleDistance = Dist;
+			}
+		}
+	}
+
+	if (MinVisibleDistance <= 0.0f)
+	{
+		return FReply::Handled();
+	}
+
+	TSharedPtr<FHLODCompareViewportClient> SourceClient = SourceViewport.IsValid() ? SourceViewport->GetCompareViewportClient() : nullptr;
+	TSharedPtr<FHLODCompareViewportClient> HLODClient = HLODViewport.IsValid() ? HLODViewport->GetCompareViewportClient() : nullptr;
+
+	if (!SourceClient.IsValid())
+	{
+		return FReply::Handled();
+	}
+
+	// Move camera along its current view direction to MinVisibleDistance from the source box surface,
+	// looking at the HLOD bounds center
+	const FVector CurrentLocation = SourceClient->GetViewLocation();
+	const FVector Direction = (CurrentLocation - SourceCenter).GetSafeNormal();
+	const float BoundsRadius = SourceBounds.GetExtent().Size();
+	const FVector NewLocation = SourceCenter + Direction * (BoundsRadius + MinVisibleDistance);
+	const FRotator NewRotation = (LookAtCenter - NewLocation).Rotation();
+
+	SourceClient->SetViewLocation(NewLocation);
+	SourceClient->SetViewRotation(NewRotation);
+	SourceClient->SetLookAtLocation(LookAtCenter);
+
+	if (HLODClient.IsValid())
+	{
+		HLODClient->SetViewLocation(NewLocation);
+		HLODClient->SetViewRotation(NewRotation);
+		HLODClient->SetLookAtLocation(LookAtCenter);
+	}
+
+	return FReply::Handled();
+}
+
+#undef LOCTEXT_NAMESPACE
diff --git a/Engine/Source/Editor/WorldPartitionEditor/Private/WorldPartition/HLOD/SHLODCompareWindow.h b/Engine/Source/Editor/WorldPartitionEditor/Private/WorldPartition/HLOD/SHLODCompareWindow.h
new file mode 100644
--- /dev/null
+++ b/Engine/Source/Editor/WorldPartitionEditor/Private/WorldPartition/HLOD/SHLODCompareWindow.h
@@ -0,0 +1,69 @@
+// Copyright Epic Games, Inc. All Rights Reserved.
+
+#pragma once
+
+#include "Widgets/SCompoundWidget.h"
+#include "Engine/EngineBaseTypes.h"
+#include "Framework/Commands/UICommandList.h"
+
+class AWorldPartitionHLOD;
+class FPreviewScene;
+class SHLODCompareViewport;
+class FHLODCompareViewportClient;
+class SHLODClipWrapper;
+class SBox;
+class STextBlock;
+class UActorComponent;
+
+class SHLODCompareWindow : public SCompoundWidget
+{
+public:
+	SLATE_BEGIN_ARGS(SHLODCompareWindow) {}
+		SLATE_ARGUMENT(TArray<TWeakObjectPtr<AWorldPartitionHLOD>>, HLODActors)
+	SLATE_END_ARGS()
+
+	void Construct(const FArguments& InArgs);
+	~SHLODCompareWindow();
+
+	virtual FReply OnKeyDown(const FGeometry& MyGeometry, const FKeyEvent& InKeyEvent) override;
+	virtual bool SupportsKeyboardFocus() const override { return true; }
+
+	void OnCameraMoved(FHLODCompareViewportClient* Source);
+	void OnViewModeChanged(FHLODCompareViewportClient* Source);
+
+private:
+	void BindCommands();
+	void SetCompareViewMode(EViewModeIndex ViewMode, FName BufferVizMode = NAME_None);
+	bool PopulateSourceScene();
+	bool PopulateHLODScene();
+	void FocusCamerasOnContent();
+	void CopyWorldLighting();
+	void OnWipePositionChanged(float NewPosition);
+	FReply OnGoToMinVisibleDistance();
+
+	TArray<TWeakObjectPtr<AWorldPartitionHLOD>> HLODActors;
+
+	TSharedPtr<FPreviewScene> SourcePreviewScene;
+	TSharedPtr<FPreviewScene> HLODPreviewScene;
+
+	TSharedPtr<SHLODCompareViewport> SourceViewport;
+	TSharedPtr<SHLODCompareViewport> HLODViewport;
+
+	TArray<UActorComponent*> SourceSceneComponents;
+	TArray<UActorComponent*> HLODSceneComponents;
+
+	float GetCameraDistanceToBounds() const;
+	float GetSourceMinDrawDistance() const;
+	float GetHLODMinVisibleDistance() const;
+	FString GetSourceZoneLabel() const;
+	FString GetHLODZoneLabel() const;
+
+	TSharedPtr<STextBlock> StatusText;
+
+	TSharedPtr<FUICommandList> CommandList;
+
+	// Wipe overlay
+	float WipePosition = 0.5f;
+	TSharedPtr<SHLODClipWrapper> ClipWrapper;
+	TSharedPtr<SBox> ViewModeToolbarContainer;
+};
diff --git a/Engine/Source/Editor/WorldPartitionEditor/Private/WorldPartition/HLOD/SHLODDistanceRuler.cpp b/Engine/Source/Editor/WorldPartitionEditor/Private/WorldPartition/HLOD/SHLODDistanceRuler.cpp
new file mode 100644
--- /dev/null
+++ b/Engine/Source/Editor/WorldPartitionEditor/Private/WorldPartition/HLOD/SHLODDistanceRuler.cpp
@@ -0,0 +1,184 @@
+// Copyright Epic Games, Inc. All Rights Reserved.
+
+#include "WorldPartition/HLOD/SHLODDistanceRuler.h"
+#include "Rendering/DrawElements.h"
+#include "Styling/AppStyle.h"
+#include "Styling/CoreStyle.h"
+#include "Fonts/FontMeasure.h"
+#include "Framework/Application/SlateApplication.h"
+
+#define LOCTEXT_NAMESPACE "SHLODDistanceRuler"
+
+static constexpr float RulerHeight = 40.0f;
+static constexpr float CameraIconSize = 12.0f;
+static constexpr float CameraRowY = 2.0f;
+static constexpr float TrackY = CameraRowY + CameraIconSize + 3.0f;
+static constexpr float TrackHeight = 3.0f;
+static constexpr float MarkerHeight = 10.0f;
+static constexpr float HorizontalPadding = 16.0f;
+
+// Colors: green = correct range (source visible), grey = too close
+static const FLinearColor CorrectRangeColor(0.1f, 0.5f, 0.1f, 0.8f);
+static const FLinearColor CorrectRangeMarkerColor(0.4f, 0.8f, 0.4f, 1.0f);
+static const FLinearColor TooCloseColor(0.25f, 0.25f, 0.25f, 0.8f);
+static const FLinearColor TooCloseMarkerColor(0.5f, 0.5f, 0.5f, 1.0f);
+
+void SHLODDistanceRuler::Construct(const FArguments& InArgs)
+{
+	SourceMinDrawDistance = InArgs._SourceMinDrawDistance;
+	HLODMinVisibleDistance = InArgs._HLODMinVisibleDistance;
+	CameraDistance = InArgs._CameraDistance;
+	SourceZoneLabel = InArgs._SourceZoneLabel;
+	HLODZoneLabel = InArgs._HLODZoneLabel;
+	WorldToMeters = InArgs._WorldToMeters;
+}
+
+FVector2D SHLODDistanceRuler::ComputeDesiredSize(float) const
+{
+	return FVector2D(100.0f, RulerHeight);
+}
+
+int32 SHLODDistanceRuler::OnPaint(const FPaintArgs& Args, const FGeometry& AllottedGeometry, const FSlateRect& MyCullingRect, FSlateWindowElementList& OutDrawElements, int32 LayerId, const FWidgetStyle& InWidgetStyle, bool bParentEnabled) const
+{
+	const float Width = AllottedGeometry.GetLocalSize().X;
+	const float TrackLeft = HorizontalPadding;
+	const float TrackRight = Width - HorizontalPadding;
+	const float TrackWidth = TrackRight - TrackLeft;
+
+	if (TrackWidth <= 0.0f)
+	{
+		return LayerId;
+	}
+
+	const float SourceDist = SourceMinDrawDistance.Get();
+	const float HLODDist = HLODMinVisibleDistance.Get();
+	const float CamDist = CameraDistance.Get();
+
+	// The ruler range: show from 0 to a bit beyond HLODMinVisibleDistance
+	const float MaxRange = FMath::Max(HLODDist * 1.3f, FMath::Max(SourceDist, 1.0f) * 3.0f);
+
+	auto DistToX = [&](float Dist) -> float
+	{
+		return TrackLeft + FMath::Clamp(Dist / MaxRange, 0.0f, 1.0f) * TrackWidth;
+	};
+
+	const FSlateBrush* WhiteBrush = FCoreStyle::Get().GetBrush("WhiteBrush");
+	const FSlateFontInfo Font = FCoreStyle::GetDefaultFontStyle("Regular", 8);
+	const TSharedRef<FSlateFontMeasure> FontMeasure = FSlateApplication::Get().GetRenderer()->GetFontMeasureService();
+
+	// Draw the track background
+	FSlateDrawElement::MakeBox(
+		OutDrawElements, LayerId,
+		AllottedGeometry.ToPaintGeometry(FVector2D(TrackWidth, TrackHeight), FSlateLayoutTransform(FVector2D(TrackLeft, TrackY))),
+		WhiteBrush, ESlateDrawEffect::None, FLinearColor(0.15f, 0.15f, 0.15f, 1.0f));
+
+	// Source zone (0 to HLODMinVisibleDistance) - grey, too close for meaningful HLOD comparison
+	{
+		float ZoneRight = DistToX(HLODDist);
+		float ZoneWidth = ZoneRight - TrackLeft;
+		if (ZoneWidth > 0.0f)
+		{
+			FSlateDrawElement::MakeBox(
+				OutDrawElements, LayerId + 1,
+				AllottedGeometry.ToPaintGeometry(FVector2D(ZoneWidth, TrackHeight), FSlateLayoutTransform(FVector2D(TrackLeft, TrackY))),
+				WhiteBrush, ESlateDrawEffect::None, TooCloseColor);
+		}
+	}
+
+	// HLOD zone (HLODMinVisibleDistance onward) - green, the distance at which the HLOD is seen in-game
+	if (HLODDist > 0.0f)
+	{
+		float ZoneLeft = DistToX(HLODDist);
+		float ZoneWidth = TrackRight - ZoneLeft;
+		if (ZoneWidth > 0.0f)
+		{
+			FSlateDrawElement::MakeBox(
+				OutDrawElements, LayerId + 1,
+				AllottedGeometry.ToPaintGeometry(FVector2D(ZoneWidth, TrackHeight), FSlateLayoutTransform(FVector2D(ZoneLeft, TrackY))),
+				WhiteBrush, ESlateDrawEffect::None, CorrectRangeColor);
+		}
+	}
+
+	// Camera indicator — above the track, colored by zone
+	{
+		float CamX = DistToX(CamDist);
+		const FLinearColor CamColor = (CamDist >= HLODDist) ? CorrectRangeMarkerColor : TooCloseMarkerColor;
+
+		// Camera icon
+		const FSlateBrush* CameraBrush = FAppStyle::Get().GetBrush("ClassIcon.CameraComponent");
+		float IconX = CamX - CameraIconSize * 0.5f;
+		FSlateDrawElement::MakeBox(
+			OutDrawElements, LayerId + 3,
+			AllottedGeometry.ToPaintGeometry(FVector2D(CameraIconSize, CameraIconSize), FSlateLayoutTransform(FVector2D(IconX, CameraRowY))),
+			CameraBrush, ESlateDrawEffect::None, CamColor);
+
+		// Camera distance label in meters
+		const float W2M = FMath::Max(WorldToMeters.Get(), 1.0f);
+		FString CamLabel = FString::Printf(TEXT("(%.1fm)"), CamDist / W2M);
+		FVector2D LabelSize = FontMeasure->Measure(CamLabel, Font);
+		float LabelX = FMath::Clamp(CamX + CameraIconSize * 0.5f + 3.0f, TrackLeft, TrackRight - LabelSize.X);
+		FSlateDrawElement::MakeText(
+			OutDrawElements, LayerId + 3,
+			AllottedGeometry.ToPaintGeometry(LabelSize, FSlateLayoutTransform(FVector2D(LabelX, CameraRowY + (CameraIconSize - LabelSize.Y) * 0.5f))),
+			CamLabel, Font, ESlateDrawEffect::None, CamColor);
+	}
+
+	// Markers and labels — below the track
+	const float MarkerTopY = TrackY + TrackHeight;
+	const float LabelY = MarkerTopY + MarkerHeight + 1.0f;
+
+	// Source MinDrawDistance marker (grey)
+	if (SourceDist > 0.0f)
+	{
+		float MarkerX = DistToX(SourceDist);
+		FSlateDrawElement::MakeBox(
+			OutDrawElements, LayerId + 2,
+			AllottedGeometry.ToPaintGeometry(FVector2D(1.0f, MarkerHeight), FSlateLayoutTransform(FVector2D(MarkerX, MarkerTopY))),
+			WhiteBrush, ESlateDrawEffect::None, TooCloseMarkerColor);
+	}
+
+	// HLOD MinVisibleDistance marker (green)
+	if (HLODDist > 0.0f)
+	{
+		float MarkerX = DistToX(HLODDist);
+		FSlateDrawElement::MakeBox(
+			OutDrawElements, LayerId + 2,
+			AllottedGeometry.ToPaintGeometry(FVector2D(1.0f, MarkerHeight), FSlateLayoutTransform(FVector2D(MarkerX, MarkerTopY))),
+			WhiteBrush, ESlateDrawEffect::None, CorrectRangeMarkerColor);
+	}
+
+	// Source zone label — at SourceMinDrawDistance marker position (grey)
+	{
+		FString Label = SourceZoneLabel.Get();
+		if (!Label.IsEmpty())
+		{
+			float MarkerX = DistToX(SourceDist);
+			FVector2D LabelSize = FontMeasure->Measure(Label, Font);
+			float LabelX = FMath::Clamp(MarkerX - LabelSize.X * 0.5f, TrackLeft, TrackRight - LabelSize.X);
+			FSlateDrawElement::MakeText(
+				OutDrawElements, LayerId + 2,
+				AllottedGeometry.ToPaintGeometry(LabelSize, FSlateLayoutTransform(FVector2D(LabelX, LabelY))),
+				Label, Font, ESlateDrawEffect::None, TooCloseMarkerColor);
+		}
+	}
+
+	// HLOD zone label — at HLOD MinVisibleDistance marker position (green)
+	if (HLODDist > 0.0f)
+	{
+		FString Label = HLODZoneLabel.Get();
+		if (!Label.IsEmpty())
+		{
+			float MarkerX = DistToX(HLODDist);
+			FVector2D LabelSize = FontMeasure->Measure(Label, Font);
+			float LabelX = FMath::Clamp(MarkerX - LabelSize.X * 0.5f, TrackLeft, TrackRight - LabelSize.X);
+			FSlateDrawElement::MakeText(
+				OutDrawElements, LayerId + 2,
+				AllottedGeometry.ToPaintGeometry(LabelSize, FSlateLayoutTransform(FVector2D(LabelX, LabelY))),
+				Label, Font, ESlateDrawEffect::None, CorrectRangeMarkerColor);
+		}
+	}
+
+	return LayerId + 4;
+}
+
+#undef LOCTEXT_NAMESPACE
diff --git a/Engine/Source/Editor/WorldPartitionEditor/Private/WorldPartition/HLOD/SHLODDistanceRuler.h b/Engine/Source/Editor/WorldPartitionEditor/Private/WorldPartition/HLOD/SHLODDistanceRuler.h
new file mode 100644
--- /dev/null
+++ b/Engine/Source/Editor/WorldPartitionEditor/Private/WorldPartition/HLOD/SHLODDistanceRuler.h
@@ -0,0 +1,44 @@
+// Copyright Epic Games, Inc. All Rights Reserved.
+
+#pragma once
+
+#include "Widgets/SLeafWidget.h"
+
+/**
+ * A horizontal ruler showing camera distance relative to source actors and HLOD transition.
+ *
+ * Layout:
+ *   |--- Source visible ---|- gap -|--- HLOD visible ---|
+ *   0    SourceMinDraw     ^       HLODMinVisible       far
+ *        "Actors"          camera  "HLOD1 (128.0m)"
+ */
+class SHLODDistanceRuler : public SLeafWidget
+{
+public:
+	SLATE_BEGIN_ARGS(SHLODDistanceRuler)
+		: _SourceMinDrawDistance(0.0f)
+		, _HLODMinVisibleDistance(0.0f)
+		, _CameraDistance(0.0f)
+	{}
+		SLATE_ATTRIBUTE(float, SourceMinDrawDistance)
+		SLATE_ATTRIBUTE(float, HLODMinVisibleDistance)
+		SLATE_ATTRIBUTE(float, CameraDistance)
+		SLATE_ATTRIBUTE(FString, SourceZoneLabel)
+		SLATE_ATTRIBUTE(FString, HLODZoneLabel)
+		SLATE_ATTRIBUTE(float, WorldToMeters)
+	SLATE_END_ARGS()
+
+	void Construct(const FArguments& InArgs);
+
+	// SWidget interface
+	virtual int32 OnPaint(const FPaintArgs& Args, const FGeometry& AllottedGeometry, const FSlateRect& MyCullingRect, FSlateWindowElementList& OutDrawElements, int32 LayerId, const FWidgetStyle& InWidgetStyle, bool bParentEnabled) const override;
+	virtual FVector2D ComputeDesiredSize(float) const override;
+
+private:
+	TAttribute<float> SourceMinDrawDistance;
+	TAttribute<float> HLODMinVisibleDistance;
+	TAttribute<float> CameraDistance;
+	TAttribute<FString> SourceZoneLabel;
+	TAttribute<FString> HLODZoneLabel;
+	TAttribute<float> WorldToMeters;
+};
diff --git a/Engine/Source/Editor/WorldPartitionEditor/Private/WorldPartition/HLOD/SHLODWipeHandle.cpp b/Engine/Source/Editor/WorldPartitionEditor/Private/WorldPartition/HLOD/SHLODWipeHandle.cpp
new file mode 100644
--- /dev/null
+++ b/Engine/Source/Editor/WorldPartitionEditor/Private/WorldPartition/HLOD/SHLODWipeHandle.cpp
@@ -0,0 +1,92 @@
+// Copyright Epic Games, Inc. All Rights Reserved.
+
+#include "WorldPartition/HLOD/SHLODWipeHandle.h"
+#include "Rendering/DrawElements.h"
+#include "Styling/CoreStyle.h"
+
+#define LOCTEXT_NAMESPACE "SHLODWipeHandle"
+
+static constexpr float WipeHandleHitWidth = 12.0f;
+static constexpr float WipeLineVisualWidth = 2.0f;
+
+void SHLODWipeHandle::Construct(const FArguments& InArgs)
+{
+	OnWipePositionChanged = InArgs._OnWipePositionChanged;
+}
+
+int32 SHLODWipeHandle::OnPaint(const FPaintArgs& Args, const FGeometry& AllottedGeometry, const FSlateRect& MyCullingRect, FSlateWindowElementList& OutDrawElements, int32 LayerId, const FWidgetStyle& InWidgetStyle, bool bParentEnabled) const
+{
+	const FVector2D LocalSize = AllottedGeometry.GetLocalSize();
+
+	// Draw a 2px vertical line centered in the hit area
+	const float LineOffset = (LocalSize.X - WipeLineVisualWidth) * 0.5f;
+
+	FSlateDrawElement::MakeBox(
+		OutDrawElements,
+		LayerId,
+		AllottedGeometry.ToPaintGeometry(FVector2D(WipeLineVisualWidth, LocalSize.Y), FSlateLayoutTransform(FVector2D(LineOffset, 0.0f))),
+		FCoreStyle::Get().GetBrush("WhiteBrush"),
+		ESlateDrawEffect::None,
+		FLinearColor(1.0f, 1.0f, 1.0f, 0.7f)
+	);
+
+	return LayerId + 1;
+}
+
+FVector2D SHLODWipeHandle::ComputeDesiredSize(float) const
+{
+	return FVector2D(WipeHandleHitWidth, 1.0f);
+}
+
+FReply SHLODWipeHandle::OnMouseButtonDown(const FGeometry& MyGeometry, const FPointerEvent& MouseEvent)
+{
+	if (MouseEvent.GetEffectingButton() == EKeys::LeftMouseButton)
+	{
+		bDragging = true;
+		return FReply::Handled().CaptureMouse(SharedThis(this));
+	}
+	return FReply::Unhandled();
+}
+
+FReply SHLODWipeHandle::OnMouseButtonUp(const FGeometry& MyGeometry, const FPointerEvent& MouseEvent)
+{
+	if (MouseEvent.GetEffectingButton() == EKeys::LeftMouseButton && bDragging)
+	{
+		bDragging = false;
+		return FReply::Handled().ReleaseMouseCapture();
+	}
+	return FReply::Unhandled();
+}
+
+FReply SHLODWipeHandle::OnMouseMove(const FGeometry& MyGeometry, const FPointerEvent& MouseEvent)
+{
+	if (bDragging)
+	{
+		// Find the parent SCanvas (or whatever contains us) to compute relative position
+		// The wipe handle is placed inside an SCanvas that fills the entire overlay area.
+		// We need to compute the wipe position relative to that parent canvas.
+		TSharedPtr<SWidget> ParentWidget = GetParentWidget();
+		if (ParentWidget.IsValid())
+		{
+			FGeometry ParentGeometry = ParentWidget->GetTickSpaceGeometry();
+			FVector2D LocalMousePos = ParentGeometry.AbsoluteToLocal(MouseEvent.GetScreenSpacePosition());
+			float ParentWidth = ParentGeometry.GetLocalSize().X;
+
+			if (ParentWidth > 0.0f)
+			{
+				float NewPosition = FMath::Clamp(LocalMousePos.X / ParentWidth, 0.0f, 1.0f);
+				OnWipePositionChanged.ExecuteIfBound(NewPosition);
+			}
+		}
+
+		return FReply::Handled();
+	}
+	return FReply::Unhandled();
+}
+
+FCursorReply SHLODWipeHandle::OnCursorQuery(const FGeometry& MyGeometry, const FPointerEvent& CursorEvent) const
+{
+	return FCursorReply::Cursor(EMouseCursor::ResizeLeftRight);
+}
+
+#undef LOCTEXT_NAMESPACE
diff --git a/Engine/Source/Editor/WorldPartitionEditor/Private/WorldPartition/HLOD/SHLODWipeHandle.h b/Engine/Source/Editor/WorldPartitionEditor/Private/WorldPartition/HLOD/SHLODWipeHandle.h
new file mode 100644
--- /dev/null
+++ b/Engine/Source/Editor/WorldPartitionEditor/Private/WorldPartition/HLOD/SHLODWipeHandle.h
@@ -0,0 +1,33 @@
+// Copyright Epic Games, Inc. All Rights Reserved.
+
+#pragma once
+
+#include "Widgets/SLeafWidget.h"
+
+DECLARE_DELEGATE_OneParam(FOnWipePositionChanged, float);
+
+/**
+ * A draggable vertical wipe handle for the HLOD compare overlay.
+ * Draws a thin vertical line and handles mouse drag to update wipe position.
+ */
+class SHLODWipeHandle : public SLeafWidget
+{
+public:
+	SLATE_BEGIN_ARGS(SHLODWipeHandle) {}
+		SLATE_EVENT(FOnWipePositionChanged, OnWipePositionChanged)
+	SLATE_END_ARGS()
+
+	void Construct(const FArguments& InArgs);
+
+	// SWidget interface
+	virtual int32 OnPaint(const FPaintArgs& Args, const FGeometry& AllottedGeometry, const FSlateRect& MyCullingRect, FSlateWindowElementList& OutDrawElements, int32 LayerId, const FWidgetStyle& InWidgetStyle, bool bParentEnabled) const override;
+	virtual FVector2D ComputeDesiredSize(float) const override;
+	virtual FReply OnMouseButtonDown(const FGeometry& MyGeometry, const FPointerEvent& MouseEvent) override;
+	virtual FReply OnMouseButtonUp(const FGeometry& MyGeometry, const FPointerEvent& MouseEvent) override;
+	virtual FReply OnMouseMove(const FGeometry& MyGeometry, const FPointerEvent& MouseEvent) override;
+	virtual FCursorReply OnCursorQuery(const FGeometry& MyGeometry, const FPointerEvent& CursorEvent) const override;
+
+private:
+	FOnWipePositionChanged OnWipePositionChanged;
+	bool bDragging = false;
+};
diff --git a/Engine/Source/Editor/WorldPartitionEditor/Private/WorldPartitionEditorModule.cpp b/Engine/Source/Editor/WorldPartitionEditor/Private/WorldPartitionEditorModule.cpp
--- a/Engine/Source/Editor/WorldPartitionEditor/Private/WorldPartitionEditorModule.cpp
+++ b/Engine/Source/Editor/WorldPartitionEditor/Private/WorldPartitionEditorModule.cpp
@@ -24,6 +24,7 @@
 #include "WorldPartition/HLOD/HLODCreationFilter.h"
 #include "WorldPartition/HLOD/SWorldPartitionBuildHLODsDialog.h"
 #include "WorldPartition/HLOD/HLODEditorSubsystem.h"
+#include "WorldPartition/HLOD/SHLODCompareWindow.h"
 #include "WorldPartition/WorldPartitionClassDescRegistry.h"
 #include "WorldPartition/DataLayer/ExternalDataLayerUID.h"
 #include "WorldPartition/DataLayer/ExternalDataLayerHelper.h"
@@ -125,7 +126,7 @@ static bool CanLoadUnloadSelectedVolumes(TArray<TWeakObjectPtr<AActor>> Volumes,
 	return false;
 }
 
-static void CreateLevelViewportContextMenuEntries(FMenuBuilder& MenuBuilder, TArray<TWeakObjectPtr<AActor>> Volumes, FBox SelectionBox, TSharedPtr<FHLODCreationVolumeFilter> VolumeFilter)
+static void CreateLevelViewportContextMenuEntries(FMenuBuilder& MenuBuilder, TArray<TWeakObjectPtr<AActor>> Volumes, FBox SelectionBox, TSharedPtr<FHLODCreationVolumeFilter> VolumeFilter, TArray<TWeakObjectPtr<AWorldPartitionHLOD>> CompareHLODActors)
 {
 	MenuBuilder.BeginSection("WorldPartition", LOCTEXT("WorldPartition", "World Partition"));
 	
@@ -212,7 +213,36 @@ static void CreateLevelViewportContextMenuEntries(FMenuBuilder& MenuBuilder, TAr
 
 		MenuBuilder.AddMenuEntry(LOCTEXT("BuildHLODsFromSelection", "Build HLODs From Selection"), LOCTEXT("BuildHLODSFromSelection_Tooltip", "Build HLODs from selected actor(s) bounds"), FSlateIcon(), BuildHLODsFromSelection);
 	}
-	
+
+	// Compare HLOD against source actors (only when all selected actors are HLOD actors)
+	if (!CompareHLODActors.IsEmpty())
+	{
+		FUIAction CompareHLODAction(
+			FExecuteAction::CreateLambda([CompareHLODActors]()
+			{
+				FText WindowTitle = (CompareHLODActors.Num() == 1 && CompareHLODActors[0].IsValid())
+					? FText::Format(LOCTEXT("CompareWindowTitleSingle", "Compare HLOD: {0}"), FText::FromString(CompareHLODActors[0]->GetActorLabel()))
+					: FText::Format(LOCTEXT("CompareWindowTitleMulti", "Compare HLOD ({0} HLOD actors)"), FText::AsNumber(CompareHLODActors.Num()));
+
+				TSharedRef<SWindow> Window = SNew(SWindow)
+					.Title(WindowTitle)
+					.ClientSize(FVector2D(1280, 720))
+					.SizingRule(ESizingRule::UserSized)
+					.SupportsMinimize(true)
+					.SupportsMaximize(true);
+
+				Window->SetContent(
+					SNew(SHLODCompareWindow)
+					.HLODActors(CompareHLODActors)
+				);
+
+				FSlateApplication::Get().AddWindow(Window);
+			})
+		);
+
+		MenuBuilder.AddMenuEntry(LOCTEXT("CompareHLODAgainstSource", "Compare HLOD against source"), LOCTEXT("CompareHLODAgainstSource_Tooltip", "Open a wipe comparison of the selected HLODs against their source actors"), FSlateIcon(), CompareHLODAction);
+	}
+
 	MenuBuilder.EndSection();
 }
 
@@ -243,13 +273,32 @@ static TSharedRef<FExtender> OnExtendLevelEditorMenu(const TSharedRef<FUICommand
 
 	TSharedPtr<FHLODCreationVolumeFilter> VolumeFilter = MakeShared<FHLODCreationVolumeFilter>(MoveTemp(FilterVolumes));
 
-	if (!Volumes.IsEmpty() || BoundsBuilder.IsValid() || VolumeFilter->IsActive())
+	// Collect HLOD actors for the compare entry (only if ALL selected actors are HLOD actors)
+	TArray<TWeakObjectPtr<AWorldPartitionHLOD>> CompareHLODActors;
+	bool bAllAreHLOD = !SelectedActors.IsEmpty();
+	for (AActor* Actor : SelectedActors)
+	{
+		if (AWorldPartitionHLOD* HLODActor = Cast<AWorldPartitionHLOD>(Actor))
+		{
+			CompareHLODActors.Add(HLODActor);
+		}
+		else
+		{
+			bAllAreHLOD = false;
+		}
+	}
+	if (!bAllAreHLOD)
+	{
+		CompareHLODActors.Empty();
+	}
+
+	if (!Volumes.IsEmpty() || BoundsBuilder.IsValid() || VolumeFilter->IsActive() || !CompareHLODActors.IsEmpty())
 	{
 		Extender->AddMenuExtension(
 			"ActorTypeTools",
 			EExtensionHook::After,
 			nullptr,
-			FMenuExtensionDelegate::CreateStatic(&CreateLevelViewportContextMenuEntries, Volumes, FBoxSphereBounds(BoundsBuilder).GetBox(), VolumeFilter));
+			FMenuExtensionDelegate::CreateStatic(&CreateLevelViewportContextMenuEntries, Volumes, FBoxSphereBounds(BoundsBuilder).GetBox(), VolumeFilter, CompareHLODActors));
 	}
 
 	return Extender;
diff --git a/Engine/Source/Editor/WorldPartitionEditor/WorldPartitionEditor.Build.cs b/Engine/Source/Editor/WorldPartitionEditor/WorldPartitionEditor.Build.cs
--- a/Engine/Source/Editor/WorldPartitionEditor/WorldPartitionEditor.Build.cs
+++ b/Engine/Source/Editor/WorldPartitionEditor/WorldPartitionEditor.Build.cs
@@ -10,6 +10,7 @@ public class WorldPartitionEditor : ModuleRules
         PrivateDependencyModuleNames.AddRange(
             new string[] {
 				"ApplicationCore",
+				"CommonMenuExtensions",
 				"ContentBrowser",
 				"ContentBrowserData",
 				"Core",
